Rad Audio Codec: SDK and UE integration.

This has several preflights that aren't listed, covering various platforms.

diff --git a/Engine/Config/BaseEngine.ini b/Engine/Config/BaseEngine.ini
--- a/Engine/Config/BaseEngine.ini
+++ b/Engine/Config/BaseEngine.ini
@@ -1548,6 +1548,7 @@ UseAudioThread=true
 EnableAudioMixer=false
 DeferStartupPrecache=false
 +AllWaveFormats=BINKA
++AllWaveFormats=RADA
 +AllWaveFormats=ADPCM
 +AllWaveFormats=PCM
 +AllWaveFormats=OPUS
@@ -1556,10 +1557,13 @@ DeferStartupPrecache=false
 +FormatModuleHints=AudioFormatOGG
 +FormatModuleHints=AudioFormatADPCM
 +FormatModuleHints=AudioFormatBINK
++FormatModuleHints=AudioFormatRAD
 +AudioInfoModules=OpusAudioDecoder
 +AudioInfoModules=VorbisAudioDecoder
 +AudioInfoModules=AdpcmAudioDecoder
 +AudioInfoModules=BinkAudioDecoder
++AudioInfoModules=RadAudioDecoder
+
 FallbackFormat=BINKA
 PlatformFormat=BINKA
 PlatformStreamingFormat=BINKA
diff --git a/Engine/Source/Developer/TargetPlatform/TargetPlatform.Build.cs b/Engine/Source/Developer/TargetPlatform/TargetPlatform.Build.cs
--- a/Engine/Source/Developer/TargetPlatform/TargetPlatform.Build.cs
+++ b/Engine/Source/Developer/TargetPlatform/TargetPlatform.Build.cs
@@ -59,6 +59,7 @@ public class TargetPlatform : ModuleRules
 					DynamicallyLoadedModuleNames.Add("AudioFormatOgg");
 					DynamicallyLoadedModuleNames.Add("AudioFormatOpus");
 					DynamicallyLoadedModuleNames.Add("AudioFormatBink");
+					DynamicallyLoadedModuleNames.Add("AudioFormatRad");
 				}
 			}
 			else if (Target.Platform == UnrealTargetPlatform.Mac)
@@ -72,6 +73,7 @@ public class TargetPlatform : ModuleRules
 					DynamicallyLoadedModuleNames.Add("AudioFormatOgg");
 					DynamicallyLoadedModuleNames.Add("AudioFormatOpus");
 					DynamicallyLoadedModuleNames.Add("AudioFormatBink");
+					DynamicallyLoadedModuleNames.Add("AudioFormatRad");
 				}
 
 			}
@@ -87,6 +89,7 @@ public class TargetPlatform : ModuleRules
 					DynamicallyLoadedModuleNames.Add("AudioFormatOgg");
 					DynamicallyLoadedModuleNames.Add("AudioFormatOpus");
 					DynamicallyLoadedModuleNames.Add("AudioFormatBink");
+					DynamicallyLoadedModuleNames.Add("AudioFormatRad");
 				}
 			}
 		}
diff --git a/Engine/Source/Runtime/AudioMixerCore/Public/AudioMixerTypes.h b/Engine/Source/Runtime/AudioMixerCore/Public/AudioMixerTypes.h
--- a/Engine/Source/Runtime/AudioMixerCore/Public/AudioMixerTypes.h
+++ b/Engine/Source/Runtime/AudioMixerCore/Public/AudioMixerTypes.h
@@ -53,6 +53,7 @@ namespace Audio {
 	inline FName NAME_ADPCM(TEXT("ADPCM"));
 	inline FName NAME_PCM(TEXT("PCM"));
 	inline FName NAME_OPUS(TEXT("OPUS"));
+	inline FName NAME_RADA(TEXT("RADA"));
 
 	// Not yet supported on all platforms as a selectable option so is included under "platform specific" enumeration for now. 
 	inline FName NAME_OGG(TEXT("OGG"));
diff --git a/Engine/Source/Runtime/Engine/Classes/Sound/AudioSettings.h b/Engine/Source/Runtime/Engine/Classes/Sound/AudioSettings.h
--- a/Engine/Source/Runtime/Engine/Classes/Sound/AudioSettings.h
+++ b/Engine/Source/Runtime/Engine/Classes/Sound/AudioSettings.h
@@ -67,6 +67,9 @@ enum class EDefaultAudioCompressionType : uint8
 
 	// Encodes the asset to a platform specific format and will be different depending on the platform. It does not currently support seeking.
 	PlatformSpecific,
+
+	// As BinkAudio, except better quality. Comparable CPU usage on newer platforms, higher CPU on older platforms.
+	RadAudio
 };
 
 
diff --git a/Engine/Source/Runtime/Engine/Classes/Sound/SoundWave.h b/Engine/Source/Runtime/Engine/Classes/Sound/SoundWave.h
--- a/Engine/Source/Runtime/Engine/Classes/Sound/SoundWave.h
+++ b/Engine/Source/Runtime/Engine/Classes/Sound/SoundWave.h
@@ -337,6 +337,9 @@ enum class ESoundAssetCompressionType : uint8
 
 	// The project defines the codec used for this asset.
 	ProjectDefined,
+
+	// As BinkAudio, except better quality. Comparable CPU usage.
+	RadAudio,
 };
 
 
@@ -346,6 +349,7 @@ namespace Audio
 	{
 		switch (InDecoderType)
 		{
+		case ESoundAssetCompressionType::RadAudio:				return NAME_RADA;
 		case ESoundAssetCompressionType::BinkAudio:				return NAME_BINKA;
 		case ESoundAssetCompressionType::ADPCM:					return NAME_ADPCM;
 		case ESoundAssetCompressionType::PCM:					return NAME_PCM;
@@ -362,6 +366,7 @@ namespace Audio
 	{
 		switch (InDefaultCompressionType)
 		{
+			case EDefaultAudioCompressionType::RadAudio:			return ESoundAssetCompressionType::RadAudio;
 			case EDefaultAudioCompressionType::BinkAudio:			return ESoundAssetCompressionType::BinkAudio;
 			case EDefaultAudioCompressionType::ADPCM:				return ESoundAssetCompressionType::ADPCM;
 			case EDefaultAudioCompressionType::PCM:					return ESoundAssetCompressionType::PCM;
diff --git a/Engine/Source/Runtime/Engine/Engine.Build.cs b/Engine/Source/Runtime/Engine/Engine.Build.cs
--- a/Engine/Source/Runtime/Engine/Engine.Build.cs
+++ b/Engine/Source/Runtime/Engine/Engine.Build.cs
@@ -172,6 +172,7 @@ public class Engine : ModuleRules
 		}
 
 		// Cross platform Audio Codecs: (we build here, but don't depend on them directly)
+		DynamicallyLoadedModuleNames.Add("RadAudioDecoder");
 		DynamicallyLoadedModuleNames.Add("BinkAudioDecoder");
 		DynamicallyLoadedModuleNames.Add("VorbisAudioDecoder");
 		DynamicallyLoadedModuleNames.Add("OpusAudioDecoder");
diff --git a/Engine/Source/Runtime/Engine/Private/AudioFormatSettings.cpp b/Engine/Source/Runtime/Engine/Private/AudioFormatSettings.cpp
--- a/Engine/Source/Runtime/Engine/Private/AudioFormatSettings.cpp
+++ b/Engine/Source/Runtime/Engine/Private/AudioFormatSettings.cpp
@@ -106,7 +106,7 @@ namespace Audio
 			}
 			else
 			{
-				AllWaveFormats = { NAME_BINKA, NAME_ADPCM, NAME_PCM, NAME_OPUS };
+				AllWaveFormats = { NAME_BINKA, NAME_ADPCM, NAME_PCM, NAME_OPUS, NAME_RADA};
 				UE_LOG(LogAudio, Warning, TEXT("Audio:AllWaveFormats is not defined, defaulting to built in formats. (%s)"), *MakePrettyArrayToString(AllWaveFormats));
 			}
 		}
diff --git a/Engine/Source/Runtime/RadAudioCodec/SDK/BuildAppleLibs.command b/Engine/Source/Runtime/RadAudioCodec/SDK/BuildAppleLibs.command
new file mode 100755
--- /dev/null
+++ b/Engine/Source/Runtime/RadAudioCodec/SDK/BuildAppleLibs.command
@@ -0,0 +1,496 @@
+#!/bin/bash
+# Copyright Epic Games, Inc. All Rights Reserved.
+
+
+set -e
+
+UE_MODULE_LOCATION=`pwd`
+
+RADA_SOURCE_LOCATION="$UE_MODULE_LOCATION/Src/RadA"
+RADAUDIO_SOURCE_LOCATION="$UE_MODULE_LOCATION/Src/RadAudio"
+
+BUILD_ROOT="$UE_MODULE_LOCATION/Intermediate"
+OUTPUT_LOCATION="$UE_MODULE_LOCATION/Lib"
+
+echo "It's normal to see empty symbol ranlib warnings"
+
+#if [[ ! -w $OUTPUT_LOCATION/librada_decode_osx64.a ]]; then
+#	echo "Check out $OUTPUT_LOCATION/librada_decode_osx64.a"
+#	exit 1
+#fi
+#if [[ ! -w $OUTPUT_LOCATION/libradaudio_decoder_osx.a ]]; then
+#	echo "Check out $OUTPUT_LOCATION/libradaudio_decoder_osx.a"
+#	exit 1
+#fi
+
+rm -rf $BUILD_ROOT
+mkdir $BUILD_ROOT
+
+VISIONOS_ARM64_ARGS=(
+	-arch arm64
+	--target=arm64-apple-xros1.0
+)
+
+VISIONOS_SIM_ARM64_ARGS=(
+	-arch arm64
+	--target=arm64-apple-xros1.0-simulator
+)
+
+TVOS_ARM64_ARGS=(
+	-arch arm64
+	--target=arm64-apple-tvos15
+)
+
+IOS_ARM64_ARGS=(
+	-arch arm64
+	--target=arm-apple-ios10
+)
+
+IOS_SIM_ARGS_X64=(
+	-arch x86_64
+	--target=x86_64-apple-ios10
+	-mmmx
+)
+
+IOS_SIM_ARGS_ARM64=(
+	-arch arm64
+	-miphonesimulator-version-min=15
+)
+
+OSX_X64_ARGS=(
+	-arch x86_64
+	-msse
+	-msse2
+	-msse3
+	-mssse3
+	-mmacosx-version-min=10.14
+)
+
+OSX_ARM64_ARGS=(
+	-arch arm64
+	-mmacosx-version-min=10.14
+)
+
+
+AVX2_ARGS=(
+	-mavx
+	-mavx2
+)
+
+CPP_ARGS=(
+	-std=c++14
+	)
+
+COMMON_ARGS=(
+	-c
+	-DNDEBUG
+	-DUSING_EGT	# tell rrCore.h we are using egttypes.h
+	-D__RADINSTATICLIB__
+	-DRADAUDIO_WRAP=UERA # Prefix symbols so that mulitple libs using the same source dont get reused by the linker
+	-DRADA_WRAP=UERA # Prefix symbols so that mulitple libs using the same source dont get reused by the linker
+	-ffp-contract=off # Prevent FMA contraction for consistency between arm/x64
+	-fno-exceptions
+	-fno-omit-frame-pointer
+	-fno-rtti
+	-fno-strict-aliasing # prevent optimizations from introducing random silent bugs
+	-fno-unroll-loops
+	-fno-vectorize # Weve vectorized everything so this just makes the tail computation get unrolled unnecessarily
+	-fvisibility=hidden
+	-ggdb
+	-mllvm
+	-inline-threshold=64 # Pass inline threhold to llvm to prevent binary size explosion
+	-momit-leaf-frame-pointer
+	-O2
+	-IInclude
+	-I$RADAUDIO_SOURCE_LOCATION
+)
+
+RADAUDIO_SOURCES=(
+	radaudio_decoder.c
+	radaudio_mdct.cpp
+	radaudio_mdct_neon.cpp
+	radaudio_mdct_sse2.cpp
+	cpux86.cpp
+	radaudio_decoder_sse2.c
+	radaudio_decoder_sse4.c
+	radaudio_decoder_neon.c
+)
+
+RADAUDIO_SOURCES_AVX2=(
+	radaudio_decoder_avx2.c
+	radaudio_mdct_avx2.cpp
+)
+
+RADAUDIO_ENCODER_SOURCES=(
+	radaudio_encoder.c
+	radaudio_mdct.cpp
+	radaudio_mdct_neon.cpp
+	radaudio_mdct_sse2.cpp
+	cpux86.cpp
+	radaudio_encoder_sse.c
+	radaudio_encoder_neon.c
+)
+
+RADA_ENCODE_SOURCES=(
+	rada_encode.cpp
+)
+
+RADA_SOURCES=(
+	rada_decode.cpp
+)
+
+
+
+# We expect:
+#	CLANG_ARGS
+#	CLANG_SOURCES
+#	BUILD_LOCATION
+#	CLANG_SOURCE_LOCATION
+#	CPP_ARGS
+#	COMMON_ARGS
+#	APPLE_SDK
+# We modify and update:
+#	OUTPUT
+
+function CallClang()
+{
+	local CLANG_PATH=$(xcrun --sdk $APPLE_SDK --find clang)
+	local SYSROOT=(
+		-isysroot
+		$(xcrun --sdk $APPLE_SDK --show-sdk-path)
+		)
+
+	mkdir -p $BUILD_LOCATION
+	for source_file in "${CLANG_SOURCES[@]}"
+	do
+		if [ "${source_file##*\.}" == "cpp" ]
+		then
+			$CLANG_PATH ${SYSROOT[@]} ${CPP_ARGS} ${CLANG_ARGS[@]} ${COMMON_ARGS[@]} -o $BUILD_LOCATION/${source_file%.*}.o $CLANG_SOURCE_LOCATION/$source_file
+		else
+			$CLANG_PATH ${SYSROOT[@]} ${CLANG_ARGS[@]} ${COMMON_ARGS[@]} -o $BUILD_LOCATION/${source_file%.*}.o $CLANG_SOURCE_LOCATION/$source_file
+		fi
+
+		OUTPUT+=( $BUILD_LOCATION/${source_file%.*}.o )
+	done
+}
+
+#
+#
+# ----- RADAUDIO DECODER -------------------------------------
+#
+#
+echo Rad Audio Decoder...
+CLANG_SOURCE_LOCATION=$RADAUDIO_SOURCE_LOCATION
+
+#
+# OSX
+#
+APPLE_SDK=macosx
+echo ...OSX
+# ---- OSX X64 ----
+OUTPUT=()
+BUILD_LOCATION=$BUILD_ROOT/radaudio_osx_x64
+
+CLANG_ARGS=(${OSX_X64_ARGS[@]})
+CLANG_SOURCES=(${RADAUDIO_SOURCES[@]})
+CallClang
+
+CLANG_ARGS=(${OSX_X64_ARGS[@]} ${AVX2_ARGS[@]})
+CLANG_SOURCES=(${RADAUDIO_SOURCES_AVX2[@]})
+CallClang
+
+ar rcs $OUTPUT_LOCATION/libradaudio_decoder_osx_x64_static.a ${OUTPUT[@]}
+
+# ---- OSX ARM64 ----
+OUTPUT=()
+BUILD_LOCATION=$BUILD_ROOT/radaudio_osx_arm64
+
+CLANG_ARGS=(${OSX_ARM64_ARGS[@]})
+CLANG_SOURCES=(${RADAUDIO_SOURCES[@]})
+CallClang
+
+ar rcs $OUTPUT_LOCATION/libradaudio_decoder_osx_arm64_static.a ${OUTPUT[@]}
+
+# ---- LIPO ----
+lipo -create $OUTPUT_LOCATION/libradaudio_decoder_osx_x64_static.a $OUTPUT_LOCATION/libradaudio_decoder_osx_arm64_static.a -output $OUTPUT_LOCATION/libradaudio_decoder_osx.a
+rm $OUTPUT_LOCATION/libradaudio_decoder_osx_arm64_static.a
+rm $OUTPUT_LOCATION/libradaudio_decoder_osx_x64_static.a
+
+
+#
+# iOS
+#
+echo ...IOS
+APPLE_SDK=iphoneos
+
+# ---- iOS ARM64 ----
+OUTPUT=()
+BUILD_LOCATION=$BUILD_ROOT/radaudio_ios_arm64
+
+CLANG_ARGS=(${IOS_ARM64_ARGS[@]})
+CallClang
+
+ar rcs $OUTPUT_LOCATION/libradaudio_decoder_ios.a ${OUTPUT[@]}
+
+# ---- iOS Sim x64 ----
+APPLE_SDK=iphonesimulator
+
+OUTPUT=()
+BUILD_LOCATION=$BUILD_ROOT/radaudio_iossim_x64
+
+CLANG_ARGS=(${IOS_SIM_ARGS_X64[@]})
+CallClang
+
+ar rcs $OUTPUT_LOCATION/libradaudio_decoder_iossim_x64_static.a ${OUTPUT[@]}
+
+# ---- iOS Sim ARM64 ----
+OUTPUT=()
+BUILD_LOCATION=$BUILD_ROOT/radaudio_iossim_arm64
+
+CLANG_ARGS=(${IOS_SIM_ARGS_ARM64[@]})
+CallClang
+
+ar rcs $OUTPUT_LOCATION/libradaudio_decoder_iossim_arm64_static.a ${OUTPUT[@]}
+
+# ---- LIPO SIM ----
+lipo -create $OUTPUT_LOCATION/libradaudio_decoder_iossim_x64_static.a $OUTPUT_LOCATION/libradaudio_decoder_iossim_arm64_static.a -output $OUTPUT_LOCATION/libradaudio_decoder_iossim.a
+rm $OUTPUT_LOCATION/libradaudio_decoder_iossim_x64_static.a
+rm $OUTPUT_LOCATION/libradaudio_decoder_iossim_arm64_static.a
+
+
+#
+# TVOS
+#
+echo ...TVOS
+APPLE_SDK=appletvos
+
+OUTPUT=()
+BUILD_LOCATION=$BUILD_ROOT/radaudio_tvos_arm64
+
+CLANG_ARGS=(${TVOS_ARM64_ARGS[@]})
+CallClang
+
+ar rcs $OUTPUT_LOCATION/libradaudio_decoder_tvos.a ${OUTPUT[@]}
+
+#
+# VisionOS
+#
+echo ...VisionOS
+APPLE_SDK=xros
+
+OUTPUT=()
+BUILD_LOCATION=$BUILD_ROOT/radaudio_visionos_arm64
+
+CLANG_ARGS=(${VISIONOS_ARM64_ARGS[@]})
+CallClang
+
+ar rcs $OUTPUT_LOCATION/libradaudio_decoder_visionos.a ${OUTPUT[@]}
+
+# ---- sim ----
+APPLE_SDK=xrsimulator
+
+OUTPUT=()
+BUILD_LOCATION=$BUILD_ROOT/radaudio_visionossim_arm64
+
+CLANG_ARGS=(${VISIONOS_SIM_ARM64_ARGS[@]})
+CallClang
+
+ar rcs $OUTPUT_LOCATION/libradaudio_decoder_visionossim.a ${OUTPUT[@]}
+
+
+
+#
+#
+# ----- RADAUDIO ENCODER -------------------------------------
+#
+#
+echo RadAudio Encoder...
+
+#
+# OSX
+#
+APPLE_SDK=macosx
+
+# ---- OSX X64 ----
+OUTPUT=()
+BUILD_LOCATION=$BUILD_ROOT/radaudio_encoder_osx_x64
+
+CLANG_ARGS=(${OSX_X64_ARGS[@]})
+CLANG_SOURCES=(${RADAUDIO_ENCODER_SOURCES[@]})
+CallClang
+
+ar rcs $OUTPUT_LOCATION/libradaudio_encoder_osx_x64_static.a ${OUTPUT[@]}
+
+# ---- OSX ARM64 ----
+OUTPUT=()
+BUILD_LOCATION=$BUILD_ROOT/radaudio_encoder_osx_arm64
+
+CLANG_ARGS=(${OSX_ARM64_ARGS[@]})
+CLANG_SOURCES=(${RADAUDIO_ENCODER_SOURCES[@]})
+CallClang
+
+ar rcs $OUTPUT_LOCATION/libradaudio_encoder_osx_arm64_static.a ${OUTPUT[@]}
+
+# ---- LIPO ----
+lipo -create $OUTPUT_LOCATION/libradaudio_encoder_osx_x64_static.a $OUTPUT_LOCATION/libradaudio_encoder_osx_arm64_static.a -output $OUTPUT_LOCATION/libradaudio_encoder_osx.a
+rm $OUTPUT_LOCATION/libradaudio_encoder_osx_arm64_static.a
+rm $OUTPUT_LOCATION/libradaudio_encoder_osx_x64_static.a
+
+
+#
+#
+# ----- RADA DECODE -------------------------------------
+#
+#
+CLANG_SOURCES=(${RADA_SOURCES[@]})
+CLANG_SOURCE_LOCATION=$RADA_SOURCE_LOCATION
+echo RadA Decoder
+
+#
+# OSX
+#
+echo ...OSX
+
+# ---- OSX X64 ----
+OUTPUT=()
+BUILD_LOCATION=$BUILD_ROOT/rada_osx_x64
+CLANG_ARGS=(${OSX_X64_ARGS[@]})
+CallClang
+
+ar rcs $OUTPUT_LOCATION/librada_decode_osx_x64_static.a ${OUTPUT[@]}
+
+# ---- OSX ARM64 ----
+OUTPUT=()
+BUILD_LOCATION=$BUILD_ROOT/rada_osx_arm64
+CLANG_ARGS=(${OSX_ARM64_ARGS[@]})
+CallClang
+
+ar rcs $OUTPUT_LOCATION/librada_decode_osx_arm64_static.a ${OUTPUT[@]}
+
+# ---- LIPO ----
+lipo -create $OUTPUT_LOCATION/librada_decode_osx_x64_static.a $OUTPUT_LOCATION/librada_decode_osx_arm64_static.a -output $OUTPUT_LOCATION/librada_decode_osx.a
+rm $OUTPUT_LOCATION/librada_decode_osx_x64_static.a
+rm $OUTPUT_LOCATION/librada_decode_osx_arm64_static.a
+
+
+#
+# iOS
+#
+echo ...IOS
+APPLE_SDK=iphoneos
+
+# ---- iOS ARM64 ----
+OUTPUT=()
+BUILD_LOCATION=$BUILD_ROOT/rada_ios_arm64
+
+CLANG_ARGS=(${IOS_ARM64_ARGS[@]})
+CallClang
+
+ar rcs $OUTPUT_LOCATION/librada_decode_ios.a ${OUTPUT[@]}
+
+# ---- iOS Sim x64 ----
+APPLE_SDK=iphonesimulator
+OUTPUT=()
+BUILD_LOCATION=$BUILD_ROOT/rada_iossim_x64
+
+CLANG_ARGS=(${IOS_SIM_ARGS_X64[@]})
+CallClang
+
+ar rcs $OUTPUT_LOCATION/librada_decode_iossim_x64_static.a ${OUTPUT[@]}
+
+# ---- iOS Sim ARM64 ----
+OUTPUT=()
+BUILD_LOCATION=$BUILD_ROOT/rada_iossim_arm64
+
+CLANG_ARGS=(${IOS_SIM_ARGS_ARM64[@]})
+CallClang
+
+ar rcs $OUTPUT_LOCATION/librada_decode_iossim_arm64_static.a ${OUTPUT[@]}
+
+
+# ---- LIPO SIM ----
+lipo -create $OUTPUT_LOCATION/librada_decode_iossim_x64_static.a $OUTPUT_LOCATION/librada_decode_iossim_arm64_static.a -output $OUTPUT_LOCATION/librada_decode_iossim.a
+rm $OUTPUT_LOCATION/librada_decode_iossim_x64_static.a
+rm $OUTPUT_LOCATION/librada_decode_iossim_arm64_static.a
+
+
+#
+# TVOS
+#
+echo TVOS
+APPLE_SDK=appletvos
+
+OUTPUT=()
+BUILD_LOCATION=$BUILD_ROOT/rada_tvos_arm64
+
+CLANG_ARGS=(${TVOS_ARM64_ARGS[@]})
+CallClang
+
+ar rcs $OUTPUT_LOCATION/librada_decode_tvos.a ${OUTPUT[@]}
+
+
+#
+# VisionOS
+#
+echo VisionOS
+APPLE_SDK=xros
+
+OUTPUT=()
+BUILD_LOCATION=$BUILD_ROOT/rada_visionos_arm64
+
+CLANG_ARGS=(${VISIONOS_ARM64_ARGS[@]})
+CallClang
+
+ar rcs $OUTPUT_LOCATION/librada_decode_visionos.a ${OUTPUT[@]}
+
+# ---- sim ----
+APPLE_SDK=xrsimulator
+
+OUTPUT=()
+BUILD_LOCATION=$BUILD_ROOT/rada_visionossim_arm64
+
+CLANG_ARGS=(${VISIONOS_SIM_ARM64_ARGS[@]})
+CallClang
+
+ar rcs $OUTPUT_LOCATION/librada_decode_visionossim.a ${OUTPUT[@]}
+
+
+#
+#
+# ----- RADA ENCODE -------------------------------------
+#
+#
+CLANG_SOURCES=(${RADA_ENCODE_SOURCES[@]})
+CLANG_SOURCE_LOCATION=$RADA_SOURCE_LOCATION
+
+echo RadA Encode
+
+#
+# OSX
+#
+APPLE_SDK=macosx
+
+# ---- OSX X64 ----
+OUTPUT=()
+BUILD_LOCATION=$BUILD_ROOT/rada_encode_osx_x64
+CLANG_ARGS=(${OSX_X64_ARGS[@]})
+CallClang
+
+ar rcs $OUTPUT_LOCATION/librada_encode_osx_x64_static.a ${OUTPUT[@]}
+
+# ---- OSX ARM64 ----
+OUTPUT=()
+BUILD_LOCATION=$BUILD_ROOT/rada_encode_osx_arm64
+CLANG_ARGS=(${OSX_ARM64_ARGS[@]})
+CallClang
+
+ar rcs $OUTPUT_LOCATION/librada_encode_osx_arm64_static.a ${OUTPUT[@]}
+
+# ---- LIPO ----
+lipo -create $OUTPUT_LOCATION/librada_encode_osx_x64_static.a $OUTPUT_LOCATION/librada_encode_osx_arm64_static.a -output $OUTPUT_LOCATION/librada_encode_osx.a
+rm $OUTPUT_LOCATION/librada_encode_osx_x64_static.a
+rm $OUTPUT_LOCATION/librada_encode_osx_arm64_static.a
+
+
+
+echo Done.
