Copyright detection tweak/fixups

diff --git a/Engine/Source/Runtime/RadAudioCodec/SDK/Src/RadAudio/egttypes.h b/Engine/Source/Runtime/RadAudioCodec/SDK/Src/RadAudio/egttypes.h
--- a/Engine/Source/Runtime/RadAudioCodec/SDK/Src/RadAudio/egttypes.h
+++ b/Engine/Source/Runtime/RadAudioCodec/SDK/Src/RadAudio/egttypes.h
@@ -1,6 +1,4 @@
-//===============================================
 // Copyright Epic Games Tools, LLC. All Rights Reserved.
-//===============================================
 
 #ifndef __EGTTYPESH__
 #define __EGTTYPESH__
@@ -9,8 +7,6 @@
 #define __RADTYPESH__
 #define __RADRR_COREH__ // block old rr_core
 
-#define RADCOPYRIGHT "Copyright (C) 1994-2024, Epic Games Tools, LLC"
-
 #if !defined(__RADRES__) // don't include anything for resource compiles
 
 //  __RAD32__ means at least 32 bit code (always defined)
diff --git a/Engine/Source/Runtime/RadAudioCodec/SDK/Src/RadAudio/rrCore.h b/Engine/Source/Runtime/RadAudioCodec/SDK/Src/RadAudio/rrCore.h
--- a/Engine/Source/Runtime/RadAudioCodec/SDK/Src/RadAudio/rrCore.h
+++ b/Engine/Source/Runtime/RadAudioCodec/SDK/Src/RadAudio/rrCore.h
@@ -1,6 +1,4 @@
-/// ========================================================================
 // Copyright Epic Games Tools, LLC. All Rights Reserved. Global types header file
-// ========================================================================
 
 #if !defined(__RADRR_CORE2H__)
 #define __RADRR_CORE2H__
