PVS divide by zero fix.

diff --git a/Engine/Source/Runtime/RadAudioCodec/Module/Private/RadAudioInfo.cpp b/Engine/Source/Runtime/RadAudioCodec/Module/Private/RadAudioInfo.cpp
--- a/Engine/Source/Runtime/RadAudioCodec/Module/Private/RadAudioInfo.cpp
+++ b/Engine/Source/Runtime/RadAudioCodec/Module/Private/RadAudioInfo.cpp
@@ -223,7 +223,10 @@ bool FRadAudioInfo::ParseHeader(const uint8* InSrcBufferData, uint32 InSrcBuffer
 		QualityInfo->SampleRate = RadASampleRateFromEnum(FileHeader->sample_rate);
 		QualityInfo->NumChannels = FileHeader->channels;
 		QualityInfo->SampleDataSize = FileHeader->frame_count * QualityInfo->NumChannels * sizeof(int16);
-		QualityInfo->Duration = (float)FileHeader->frame_count / QualityInfo->SampleRate;
+		if (QualityInfo->SampleRate)
+		{
+			QualityInfo->Duration = (float)FileHeader->frame_count / QualityInfo->SampleRate;
+		}
 	}
 
 	return true;
