p4 edit rad audio apple libs before building

diff --git a/Engine/Source/Runtime/RadAudioCodec/SDK/BuildAppleLibs.command b/Engine/Source/Runtime/RadAudioCodec/SDK/BuildAppleLibs.command
--- a/Engine/Source/Runtime/RadAudioCodec/SDK/BuildAppleLibs.command
+++ b/Engine/Source/Runtime/RadAudioCodec/SDK/BuildAppleLibs.command
@@ -14,14 +14,10 @@ OUTPUT_LOCATION="$UE_MODULE_LOCATION/Lib"
 
 echo "It's normal to see empty symbol ranlib warnings"
 
-#if [[ ! -w $OUTPUT_LOCATION/librada_decode_osx64.a ]]; then
-#	echo "Check out $OUTPUT_LOCATION/librada_decode_osx64.a"
-#	exit 1
-#fi
-#if [[ ! -w $OUTPUT_LOCATION/libradaudio_decoder_osx.a ]]; then
-#	echo "Check out $OUTPUT_LOCATION/libradaudio_decoder_osx.a"
-#	exit 1
-#fi
+p4 edit $OUTPUT_LOCATION/*osx.a
+p4 edit $OUTPUT_LOCATION/*ios{,sim}.a
+p4 edit $OUTPUT_LOCATION/*tvos.a
+p4 edit $OUTPUT_LOCATION/*visionos{,sim}.a
 
 rm -rf $BUILD_ROOT
 mkdir $BUILD_ROOT
