RAD Audio Looping Fixes, and also set 80 quality to 5 not 6 RAD audio quality.

diff --git a/Engine/Source/Developer/AudioFormatRad/Private/AudioFormatRad.cpp b/Engine/Source/Developer/AudioFormatRad/Private/AudioFormatRad.cpp
--- a/Engine/Source/Developer/AudioFormatRad/Private/AudioFormatRad.cpp
+++ b/Engine/Source/Developer/AudioFormatRad/Private/AudioFormatRad.cpp
@@ -17,16 +17,16 @@ namespace AudioFormatRadPrivate
 {
 	static uint8 GetCompressionLevelFromQualityIndex(const int32 InQualityIndex) 
 	{
-		// The engine default is 80, which we want to be 6. So we split a linear mapping
+		// The engine default is 80, which we want to be 5. So we split a linear mapping
 		// if we are below or above that.
 		float MappedValue = 0;
 		if (InQualityIndex <= 80)
 		{
-			MappedValue = FMath::GetMappedRangeValueClamped(FVector2d(1, 80), FVector2d(1, 6), InQualityIndex);
+			MappedValue = FMath::GetMappedRangeValueClamped(FVector2d(1, 80), FVector2d(1, 5), InQualityIndex);
 		}
 		else
 		{
-			MappedValue = FMath::GetMappedRangeValueClamped(FVector2d(80, 100), FVector2d(6, 9), InQualityIndex);
+			MappedValue = FMath::GetMappedRangeValueClamped(FVector2d(80, 100), FVector2d(5, 9), InQualityIndex);
 		}
 
 		return (uint8)MappedValue;
@@ -41,7 +41,7 @@ class FAudioFormatRad : public IAudioFormat
 	enum
 	{
 		/** Version for Rad Audio format, this becomes part of the DDC key. */
-		UE_AUDIO_RAD_VER = 1,
+		UE_AUDIO_RAD_VER = 2,
 	};
 
 public:
diff --git a/Engine/Source/Runtime/RadAudioCodec/SDK/Src/RadA/rada_encode.cpp b/Engine/Source/Runtime/RadAudioCodec/SDK/Src/RadA/rada_encode.cpp
--- a/Engine/Source/Runtime/RadAudioCodec/SDK/Src/RadA/rada_encode.cpp
+++ b/Engine/Source/Runtime/RadAudioCodec/SDK/Src/RadA/rada_encode.cpp
@@ -781,7 +781,9 @@ uint8_t EncodeRadAFile(
     size_t encoder_total_size = 0;
     for (S32 i = 0; i < StreamCount; i++)
     {
-        size_t result = radaudio_encode_create(&streams[i].encoder, streams[i].encoder_header_buffer, streams[i].channels, WavRate, Quality);
+        size_t result = radaudio_encode_create(
+            &streams[i].encoder, streams[i].encoder_header_buffer, streams[i].channels, 
+            WavRate, Quality, SeamlessLooping ? RADAUDIO_ENC_FLAG_improve_seamless_loop : 0);
         if (result == 0)
         {
             return RADA_COMPRESS_ERROR_ENCODER;
@@ -836,7 +838,7 @@ uint8_t EncodeRadAFile(
             encode_info.force_next_blocktype = block_type_for_all_streams;
 
             // If desired for looping, set up the padding so that we get seamless across the whole file.
-            if (true)
+            if (SeamlessLooping)
             {
                 if (SamplesPerChannel < 2048)
                 {
@@ -847,7 +849,7 @@ uint8_t EncodeRadAFile(
                 else if (stream->current_offset_frame < 1024)
                 {
                     // We're at the beginning, so we want the padding to be the end of the stream.
-                    encode_info.padding = stream->samples + (SamplesPerChannel - 2048);
+                    encode_info.padding = stream->samples + SamplesPerChannel - 2048*WavChannels;
                     encode_info.padding_len = 2048;
                 }
                 else // just always provide the beginning, it won't use it if it doesn't need it.
diff --git a/Engine/Source/Runtime/RadAudioCodec/SDK/Src/RadAudio/radaudio_encoder.c b/Engine/Source/Runtime/RadAudioCodec/SDK/Src/RadAudio/radaudio_encoder.c
--- a/Engine/Source/Runtime/RadAudioCodec/SDK/Src/RadAudio/radaudio_encoder.c
+++ b/Engine/Source/Runtime/RadAudioCodec/SDK/Src/RadAudio/radaudio_encoder.c
@@ -250,19 +250,35 @@ static void mdct_block(radaudio_encoder_state *es, float *coeff, int is_short_bl
    else
       for (i=clamped_start; i < clamped_end  ; ++i)  data[i] = samples[(offset-n+i)*stride];
 
+   #if 0
+   if (clamped_start > raw_start) {
+      for (i=0; i < raw_end; ++i) {
+         printf("%4d %d\n", i, (int) (data[i]*32767));
+      }
+      printf("seamless fix for channel %d ends at %d\n", pad_channel, clamped_start);
+   }
+   #endif
+
    if (clamped_end < raw_end) {
       if (info->padding_len > 0) {
          for (i=clamped_end; i < raw_end; ++i) {
             // map back into original sample space:
-            //   if i= -convert_offset, then sample_pos is 0
+            //   if i=raw_end, then sample_pos is len
+            rrAssert(clamped_end + convert_offset == len);
             int sample_pos = i+convert_offset;
-            //   now sample_pos from len+1..len+padlen maps to 0..padlen-1 in padding
+            //   now sample_pos from len..len+padlen-1 maps to 0..padlen-1 in padding
             int padding_pos = sample_pos - len;
             if (padding_pos >= 0 && padding_pos < (int) info->padding_len)
                data[i] = info->padding[padding_pos*stride+pad_channel];
             else
                data[i] = 0;
          }
+         #if 0
+         for (i=0; i < raw_end; ++i) {
+            printf("%4d %d\n", i, (int) (data[i]*32767));
+         }
+         printf("seamless fix for channel %d starts at %d\n", pad_channel, clamped_end);
+         #endif
       } else {
          for (i=clamped_end  ; i < raw_end      ; ++i)
             data[i] = 0;
@@ -2205,6 +2221,37 @@ static float band_mantissa_decay[11][2][6] =
    { { 0.18f,0.17f,0.16f,0.14f,0.13f,0.12f }, { 0.11f,0.10f,0.09f,0.08f,0.08f,0.07f }, },
 };
 
+static float band_mantissa_base_improved_short_dc[11][2][6] =
+{
+   { { 5.50f,5.50f,5.50f,5.50f,5.00f,4.50f }, { 5.25f,3.25f,3.25f,3.25f,2.95f,2.66f }, },
+   { { 5.50f,5.50f,5.50f,5.50f,5.00f,4.50f }, { 5.25f,3.25f,3.25f,3.25f,2.95f,2.66f }, },
+   { { 5.50f,5.50f,5.50f,5.50f,5.00f,4.50f }, { 5.25f,3.25f,3.25f,3.25f,2.95f,2.66f }, },
+   { { 5.50f,5.50f,5.50f,5.50f,5.00f,4.50f }, { 5.25f,3.25f,3.25f,3.25f,2.95f,2.66f }, },
+   { { 5.50f,5.50f,5.50f,5.50f,5.00f,4.50f }, { 5.25f,3.25f,3.25f,3.25f,2.95f,2.66f }, },
+   { { 5.50f,5.50f,5.50f,5.50f,5.00f,4.50f }, { 4.50f,3.25f,3.25f,3.25f,2.95f,2.66f }, },
+   { { 5.75f,5.50f,5.50f,5.50f,5.00f,4.50f }, { 5.00f,3.25f,3.25f,3.25f,2.95f,2.66f }, },
+   { { 6.00f,6.00f,6.00f,6.00f,6.00f,5.50f }, { 5.50f,3.50f,3.50f,3.50f,3.25f,3.00f }, },
+   { { 6.50f,6.50f,6.50f,6.50f,6.00f,5.50f }, { 6.00f,4.00f,4.00f,4.00f,3.75f,3.25f }, },
+   { { 7.00f,6.50f,6.50f,6.50f,6.50f,5.50f }, { 6.50f,4.00f,4.00f,4.00f,3.75f,3.25f }, },
+   { { 6.50f,5.50f,5.50f,5.50f,5.00f,4.50f }, { 7.50f,3.25f,3.25f,3.25f,2.95f,2.66f }, },
+};
+
+static float band_mantissa_decay_improved_short_dc[11][2][6] =
+{
+   { { 0.18f,0.17f,0.16f,0.14f,0.13f,0.12f }, { 0.11f,0.10f,0.09f,0.08f,0.08f,0.07f }, },
+   { { 0.18f,0.17f,0.16f,0.14f,0.13f,0.12f }, { 0.11f,0.10f,0.09f,0.08f,0.08f,0.07f }, },
+   { { 0.18f,0.17f,0.16f,0.14f,0.13f,0.12f }, { 0.11f,0.10f,0.09f,0.08f,0.08f,0.07f }, },
+   { { 0.18f,0.17f,0.16f,0.14f,0.13f,0.12f }, { 0.11f,0.10f,0.09f,0.08f,0.08f,0.07f }, },
+   { { 0.18f,0.17f,0.16f,0.14f,0.13f,0.12f }, { 0.11f,0.10f,0.09f,0.08f,0.08f,0.07f }, },
+   { { 0.18f,0.17f,0.16f,0.14f,0.13f,0.12f }, { 0.13f,0.10f,0.09f,0.08f,0.08f,0.07f }, },
+   { { 0.18f,0.17f,0.16f,0.14f,0.13f,0.12f }, { 0.15f,0.10f,0.09f,0.08f,0.08f,0.07f }, },
+   { { 0.19f,0.18f,0.17f,0.16f,0.15f,0.14f }, { 0.17f,0.11f,0.10f,0.09f,0.09f,0.08f }, },
+   { { 0.21f,0.20f,0.19f,0.18f,0.16f,0.15f }, { 0.19f,0.12f,0.10f,0.11f,0.10f,0.09f }, },
+   { { 0.22f,0.20f,0.19f,0.18f,0.16f,0.15f }, { 0.21f,0.12f,0.10f,0.11f,0.10f,0.09f }, },
+   { { 0.18f,0.17f,0.16f,0.14f,0.13f,0.12f }, { 0.11f,0.10f,0.09f,0.08f,0.08f,0.07f }, },
+};
+
+
 #define SUBBAND_PULSES_SCALE 1.0   // SUBBAND_SCALE
 static float subband_pulses_for_band[4][10][24] =
 {
@@ -2690,7 +2737,8 @@ size_t radaudio_encode_create_internal(radaudio_encoder *rae,
                                 int num_channels,      // 1..2
                                 int sample_rate,       // in HZ
                                 int qmode,             // 0..9
-                                float quality_pulse)   // used for ratesearch during development
+                                float quality_pulse,
+                                U32 flags)             // used for ratesearch during development
 {
    radaudio_encoder_state *es = (radaudio_encoder_state *) rae;
    int i, rate_mode;
@@ -2763,10 +2811,18 @@ size_t radaudio_encode_create_internal(radaudio_encoder *rae,
    compute_bias_set(&es->biases, h.bytes_bias);
 
    for (i=0; i < radaudio_rateinfo[0][rate_mode].num_bands; ++i) {
-      F32 base0  = band_mantissa_base [qmode][0][i/4];
-      F32 base1  = band_mantissa_base [qmode][1][i/4];
-      F32 decay0 = band_mantissa_decay[qmode][0][i/4];
-      F32 decay1 = band_mantissa_decay[qmode][1][i/4];
+      F32 base0,base1,decay0,decay1;
+      if (flags & RADAUDIO_ENC_FLAG_improve_seamless_loop) {
+         base0  = band_mantissa_base_improved_short_dc [qmode][0][i/4];
+         base1  = band_mantissa_base_improved_short_dc [qmode][1][i/4];
+         decay0 = band_mantissa_decay_improved_short_dc[qmode][0][i/4];
+         decay1 = band_mantissa_decay_improved_short_dc[qmode][1][i/4];
+      } else {
+         base0  = band_mantissa_base [qmode][0][i/4];
+         base1  = band_mantissa_base [qmode][1][i/4];
+         decay0 = band_mantissa_decay[qmode][0][i/4];
+         decay1 = band_mantissa_decay[qmode][1][i/4];
+      }
 
       h.mantissa_param[0][i][0] = (S8) ( base0 *   8 + 0.5);
       h.mantissa_param[1][i][0] = (S8) ( base1 *   8 + 0.5);
@@ -2814,9 +2870,9 @@ size_t radaudio_encode_create_internal(radaudio_encoder *rae,
    return unpack_length;
 }
 
-size_t radaudio_encode_create(radaudio_encoder *es, U8 header[RADAUDIO_STREAM_HEADER_MAX], int num_channels, int sample_rate, int quality)
+size_t radaudio_encode_create(radaudio_encoder *es, U8 header[RADAUDIO_STREAM_HEADER_MAX], int num_channels, int sample_rate, int quality, U32 flags)
 {
-   return radaudio_encode_create_internal(es, header, num_channels, sample_rate, quality, 0.0f);
+   return radaudio_encode_create_internal(es, header, num_channels, sample_rate, quality, 0.0f, flags);
 }
 
 #ifdef RADAUDIO_DEVELOPMENT
diff --git a/Engine/Source/Runtime/RadAudioCodec/SDK/Src/RadAudio/radaudio_encoder.h b/Engine/Source/Runtime/RadAudioCodec/SDK/Src/RadAudio/radaudio_encoder.h
--- a/Engine/Source/Runtime/RadAudioCodec/SDK/Src/RadAudio/radaudio_encoder.h
+++ b/Engine/Source/Runtime/RadAudioCodec/SDK/Src/RadAudio/radaudio_encoder.h
@@ -81,7 +81,11 @@ RADDEFFUNC size_t radaudio_encode_create( radaudio_encoder *rae,
                                           U8  header[RADAUDIO_STREAM_HEADER_MAX],
                                           int num_channels,
                                           int sample_rate,
-                                          int quality);
+                                          int quality,
+                                          U32 flags);
+
+#define RADAUDIO_ENC_FLAG_improve_seamless_loop    1  // actually boosts precision of short block low-frequency coefficients
+
 
 // radaudio_encode_block()
 //
diff --git a/Engine/Source/Runtime/RadAudioCodec/SDK/Src/RadAudio/radaudio_encoder_internal.h b/Engine/Source/Runtime/RadAudioCodec/SDK/Src/RadAudio/radaudio_encoder_internal.h
--- a/Engine/Source/Runtime/RadAudioCodec/SDK/Src/RadAudio/radaudio_encoder_internal.h
+++ b/Engine/Source/Runtime/RadAudioCodec/SDK/Src/RadAudio/radaudio_encoder_internal.h
@@ -124,7 +124,8 @@ typedef struct
 RADDEFFUNC size_t radaudio_encode_create_internal(radaudio_encoder *es,
                             unsigned char header[64], int num_channels,
                             int sample_rate, int quality,
-                            float pulse_quality_override);  // 0..100, replaces normal pulse quality values
+                            float pulse_quality_override, // 0..100, replaces normal pulse quality values
+                            U32 flags);
 
 typedef struct
 {
