Saturate floats from RAC before converting to int16

diff --git a/Engine/Source/Runtime/RadAudioCodec/Module/Private/RadAudioInfo.cpp b/Engine/Source/Runtime/RadAudioCodec/Module/Private/RadAudioInfo.cpp
--- a/Engine/Source/Runtime/RadAudioCodec/Module/Private/RadAudioInfo.cpp
+++ b/Engine/Source/Runtime/RadAudioCodec/Module/Private/RadAudioInfo.cpp
@@ -437,7 +437,17 @@ FDecodeResult FRadAudioInfo::Decode(const uint8* CompressedData, const int32 Com
 			
 			for (int32 SampleIdx = 0; SampleIdx < DecodeResult; SampleIdx++)
 			{
-				InterleaveDestination[ChannelIdx + (SampleIdx * NumChannels)] = (int16)(InBuffer[SampleIdx] * 32768.0f);
+				float InBufferFloat = InBuffer[SampleIdx] * 32768.0f;
+				if (InBufferFloat > 32767)
+				{
+					InBufferFloat = 32767;
+				}
+				else if (InBufferFloat < -32768)
+				{
+					InBufferFloat = -32768;
+				}
+
+				InterleaveDestination[ChannelIdx + (SampleIdx * NumChannels)] = (int16)InBufferFloat;
 			}
 		}
 
