Fix for RAC seeking accuracy.

#rnx

diff --git a/Engine/Source/Runtime/RadAudioCodec/Module/Private/RadAudioInfo.cpp b/Engine/Source/Runtime/RadAudioCodec/Module/Private/RadAudioInfo.cpp
--- a/Engine/Source/Runtime/RadAudioCodec/Module/Private/RadAudioInfo.cpp
+++ b/Engine/Source/Runtime/RadAudioCodec/Module/Private/RadAudioInfo.cpp
@@ -122,7 +122,7 @@ void FRadAudioInfo::SeekToFrame(uint32 InSeekTimeFrames)
 
 	// Block based codec - we start decoding on a block boundary and need
 	// to eat frames to get to our actual spot.
-	this->ConsumeFrameCount = InSeekTimeFrames - (int32)FrameAtLocation;
+	Decoder->ConsumeFrameCount = InSeekTimeFrames - (int32)FrameAtLocation;
 
 	//
 	// Here we need to set up the data we get to point at the right spot.
diff --git a/Engine/Source/Runtime/RadAudioCodec/Module/Public/RadAudioInfo.h b/Engine/Source/Runtime/RadAudioCodec/Module/Public/RadAudioInfo.h
--- a/Engine/Source/Runtime/RadAudioCodec/Module/Public/RadAudioInfo.h
+++ b/Engine/Source/Runtime/RadAudioCodec/Module/Public/RadAudioInfo.h
@@ -36,5 +36,4 @@ protected:
 	struct RadAudioDecoder* Decoder = 0;
 	TArray<uint8> RawMemory;
 	bool bErrorStateLatch = false;
-	uint32 ConsumeFrameCount = 0;
 };
