Rad Audio
* Updating source from previous fixes related to streaming
* Combining rada + radaudio libs in to just radaudio
* Adding optimized interleave/convert loop
* Fix for incorrectly testing for AVX2 existence without checking for OS support.

diff --git a/Engine/Source/Developer/AudioFormatRad/AudioFormatRad.Build.cs b/Engine/Source/Developer/AudioFormatRad/AudioFormatRad.Build.cs
--- a/Engine/Source/Developer/AudioFormatRad/AudioFormatRad.Build.cs
+++ b/Engine/Source/Developer/AudioFormatRad/AudioFormatRad.Build.cs
@@ -79,16 +79,12 @@ public class AudioFormatRad : ModuleRules
 				if (bIsWindowsLib)
 				{
 					PublicAdditionalLibraries.Add(Path.Combine(BasePath, "radaudio_decoder_" + Architecture + ".lib"));
-					PublicAdditionalLibraries.Add(Path.Combine(BasePath, "rada_decode_" + Architecture + ".lib"));
 					PublicAdditionalLibraries.Add(Path.Combine(BasePath, "radaudio_encoder_" + Architecture + ".lib"));
-					PublicAdditionalLibraries.Add(Path.Combine(BasePath, "rada_encode_" + Architecture + ".lib"));
 				}
 				else
 				{
 					PublicAdditionalLibraries.Add(Path.Combine(BasePath, "libradaudio_decoder_" + Architecture + ".a"));
-					PublicAdditionalLibraries.Add(Path.Combine(BasePath, "librada_decode_" + Architecture + ".a"));
 					PublicAdditionalLibraries.Add(Path.Combine(BasePath, "libradaudio_encoder_" + Architecture + ".a"));
-					PublicAdditionalLibraries.Add(Path.Combine(BasePath, "librada_encode_" + Architecture + ".a"));
 				}
 			}
 		}
diff --git a/Engine/Source/Developer/AudioFormatRad/Private/AudioFormatRad.cpp b/Engine/Source/Developer/AudioFormatRad/Private/AudioFormatRad.cpp
--- a/Engine/Source/Developer/AudioFormatRad/Private/AudioFormatRad.cpp
+++ b/Engine/Source/Developer/AudioFormatRad/Private/AudioFormatRad.cpp
@@ -78,6 +78,14 @@ public:
 
 	virtual bool Cook(FName InFormat, const TArray<uint8>& InSrcBuffer, FSoundQualityInfo& InQualityInfo, TArray<uint8>& OutCompressedDataStore) const override
 	{
+#ifdef RADA_BUILD_VERSION
+		if (RadAGetBuildVersion() != RADA_BUILD_VERSION)
+		{
+			UE_LOG(LogAudioFormatRad, Error, TEXT("Header/library mismatch with Rad Audio! Header: %u, Library: %u - verify sync is correct"), RADA_BUILD_VERSION, RadAGetBuildVersion());
+			return false;
+		}
+#endif
+
 		TRACE_CPUPROFILER_EVENT_SCOPE(FAudioFormatRad::Cook);
 		check(InFormat == NAME_RADA);
 
@@ -108,7 +116,7 @@ public:
 			}
 			else if (RadCompressError == RADA_COMPRESS_ERROR_CHANS)
 			{
-				UE_LOG(LogAudioFormatRad, Warning, TEXT("Channels provided: %u"), InQualityInfo.NumChannels);
+				UE_LOG(LogAudioFormatRad, Warning, TEXT("Channels provided: %u, max channels allowed: %d"), InQualityInfo.NumChannels, RADA_MAX_CHANS);
 			}
 			return false;
 		}
@@ -122,6 +130,14 @@ public:
 
 	virtual bool CookSurround(FName InFormat, const TArray<TArray<uint8> >& InSrcBuffers, FSoundQualityInfo& InQualityInfo, TArray<uint8>& OutCompressedDataStore) const override
 	{
+#ifdef RADA_BUILD_VERSION
+		if (RadAGetBuildVersion() != RADA_BUILD_VERSION)
+		{
+			UE_LOG(LogAudioFormatRad, Error, TEXT("Header/library mismatch with Rad Audio! Header: %u, Library: %u - verify sync is correct"), RADA_BUILD_VERSION, RadAGetBuildVersion());
+			return false;
+		}
+#endif
+
 		TRACE_CPUPROFILER_EVENT_SCOPE(FAudioFormatRad::CookSurround);
 		check(InFormat == NAME_RADA);
 
@@ -358,6 +374,13 @@ public:
 	}
 	virtual void StartupModule() override
 	{
+#ifdef RADA_BUILD_VERSION
+		uint32_t LibraryBuildVersion = RadAGetBuildVersion();
+		if (LibraryBuildVersion != RADA_BUILD_VERSION)
+		{
+			UE_LOG(LogAudioFormatRad, Error, TEXT("Header/library mismatch with Rad Audio! Header: %u, Library: %u - verify sync is correct"), RADA_BUILD_VERSION, LibraryBuildVersion);
+		}
+#endif
 		RadEncoder = new FAudioFormatRad();
 	}
 	virtual void ShutdownModule() override
diff --git a/Engine/Source/Runtime/RadAudioCodec/Module/Private/RadAudioInfo.cpp b/Engine/Source/Runtime/RadAudioCodec/Module/Private/RadAudioInfo.cpp
--- a/Engine/Source/Runtime/RadAudioCodec/Module/Private/RadAudioInfo.cpp
+++ b/Engine/Source/Runtime/RadAudioCodec/Module/Private/RadAudioInfo.cpp
@@ -431,6 +431,15 @@ FDecodeResult FRadAudioInfo::Decode(const uint8* CompressedData, const int32 Com
 		// as 1) the shared lib float conversion uses 32767 instead of 32768, which means that you can end up with
 		// < -1.0f in the buffer and 2) we can't interleave with SIMD as there's only VectorDeinterleave rather than
 		// interleave.
+		#ifdef RADA_HAS_INTERLEAVING
+		const float* DeinterleavedSources[RADA_MAX_CHANNELS];
+		for (int32 ChannelIdx = 0; ChannelIdx < NumChannels; ChannelIdx++)
+		{
+			DeinterleavedSources[ChannelIdx] = DeinterleavedDecodeBuffer.GetData() + RadADecodeBlock_MaxOutputFrames * ChannelIdx + DecodeResultOffset;
+		}
+		RadAInterleave(InterleaveDestination, DeinterleavedSources, NumChannels, DecodeResult);
+
+		#else
 		for (int32 ChannelIdx = 0; ChannelIdx < NumChannels; ChannelIdx++)
 		{
 			const float* InBuffer = DeinterleavedDecodeBuffer.GetData() + RadADecodeBlock_MaxOutputFrames*ChannelIdx + DecodeResultOffset;
@@ -450,6 +459,7 @@ FDecodeResult FRadAudioInfo::Decode(const uint8* CompressedData, const int32 Com
 				InterleaveDestination[ChannelIdx + (SampleIdx * NumChannels)] = (int16)InBufferFloat;
 			}
 		}
+		#endif
 
 		if (InterleaveDestination == (int16*)OutputReservoir.GetData())
 		{
diff --git a/Engine/Source/Runtime/RadAudioCodec/Module/RadAudioDecoder.Build.cs b/Engine/Source/Runtime/RadAudioCodec/Module/RadAudioDecoder.Build.cs
--- a/Engine/Source/Runtime/RadAudioCodec/Module/RadAudioDecoder.Build.cs
+++ b/Engine/Source/Runtime/RadAudioCodec/Module/RadAudioDecoder.Build.cs
@@ -130,12 +130,10 @@ public class RadAudioDecoder : ModuleRules
                     //    Path.Combine(BasePath, "rada_decode_" + Architecture + ".lib"));
 
                     PublicAdditionalLibraries.Add(Path.Combine(BasePath, "radaudio_decoder_" + Architecture + ".lib"));
-                    PublicAdditionalLibraries.Add(Path.Combine(BasePath, "rada_decode_" + Architecture + ".lib"));
                 }
                 else
                 {
                     PublicAdditionalLibraries.Add(Path.Combine(BasePath, "libradaudio_decoder_" + Architecture + ".a"));
-                    PublicAdditionalLibraries.Add(Path.Combine(BasePath, "librada_decode_" + Architecture + ".a"));
                 }
             }
         }
diff --git a/Engine/Source/Runtime/RadAudioCodec/SDK/BuildAppleLibs.command b/Engine/Source/Runtime/RadAudioCodec/SDK/BuildAppleLibs.command
--- a/Engine/Source/Runtime/RadAudioCodec/SDK/BuildAppleLibs.command
+++ b/Engine/Source/Runtime/RadAudioCodec/SDK/BuildAppleLibs.command
@@ -1,26 +1,52 @@
 #!/bin/bash
 # Copyright Epic Games, Inc. All Rights Reserved.
 
-
 set -e
 
-UE_MODULE_LOCATION=`pwd`
+# We usually build this from the RAD source tree which has different folders - detect whether we
+# are building in UE or RAD.
+if [[ "$PWD" = */audio ]]; then
+	UE_MODULE_LOCATION=`pwd`
+	RADA_SOURCE_LOCATION="$UE_MODULE_LOCATION/../../miles/src/support/rada_encode/"
+	RADAUDIO_SOURCE_LOCATION="$UE_MODULE_LOCATION/"
+
+	BUILD_ROOT="$UE_MODULE_LOCATION/Intermediate"
+	OUTPUT_LOCATION="$UE_MODULE_LOCATION/build"
+	
+	rm -rf $BUILD_ROOT
+	mkdir $BUILD_ROOT
+	
+	[ -d $OUTPUT_LOCATION ] || mkdir $OUTPUT_LOCATION
+	
+	EXTRA_INCLUDES=(
+		-I$UE_MODULE_LOCATION/../../shared/radrtl
+		
+	)
+else
+	UE_MODULE_LOCATION=`pwd`
+	
+	RADA_SOURCE_LOCATION="$UE_MODULE_LOCATION/Src/RadA"
+	RADAUDIO_SOURCE_LOCATION="$UE_MODULE_LOCATION/Src/RadAudio"
 
-RADA_SOURCE_LOCATION="$UE_MODULE_LOCATION/Src/RadA"
-RADAUDIO_SOURCE_LOCATION="$UE_MODULE_LOCATION/Src/RadAudio"
+	BUILD_ROOT="$UE_MODULE_LOCATION/Intermediate"
+	OUTPUT_LOCATION="$UE_MODULE_LOCATION/Lib"
 
-BUILD_ROOT="$UE_MODULE_LOCATION/Intermediate"
-OUTPUT_LOCATION="$UE_MODULE_LOCATION/Lib"
+	EXTRA_INCLUDES=(
+		-IInclude
+	)
+	
+	p4 edit $OUTPUT_LOCATION/*osx.a
+	p4 edit $OUTPUT_LOCATION/*ios{,sim}.a
+	p4 edit $OUTPUT_LOCATION/*tvos.a
+	p4 edit $OUTPUT_LOCATION/*visionos{,sim}.a
+
+	rm -rf $BUILD_ROOT
+	mkdir $BUILD_ROOT
+fi
 
 echo "It's normal to see empty symbol ranlib warnings"
 
-p4 edit $OUTPUT_LOCATION/*osx.a
-p4 edit $OUTPUT_LOCATION/*ios{,sim}.a
-p4 edit $OUTPUT_LOCATION/*tvos.a
-p4 edit $OUTPUT_LOCATION/*visionos{,sim}.a
 
-rm -rf $BUILD_ROOT
-mkdir $BUILD_ROOT
 
 VISIONOS_ARM64_ARGS=(
 	-arch arm64
@@ -97,7 +123,7 @@ COMMON_ARGS=(
 	-inline-threshold=64 # Pass inline threhold to llvm to prevent binary size explosion
 	-momit-leaf-frame-pointer
 	-O2
-	-IInclude
+	
 	-I$RADAUDIO_SOURCE_LOCATION
 )
 
@@ -110,6 +136,7 @@ RADAUDIO_SOURCES=(
 	radaudio_decoder_sse2.c
 	radaudio_decoder_sse4.c
 	radaudio_decoder_neon.c
+	rada_decode.cpp
 )
 
 RADAUDIO_SOURCES_AVX2=(
@@ -125,23 +152,14 @@ RADAUDIO_ENCODER_SOURCES=(
 	cpux86.cpp
 	radaudio_encoder_sse.c
 	radaudio_encoder_neon.c
-)
-
-RADA_ENCODE_SOURCES=(
 	rada_encode.cpp
 )
 
-RADA_SOURCES=(
-	rada_decode.cpp
-)
-
-
 
 # We expect:
 #	CLANG_ARGS
 #	CLANG_SOURCES
 #	BUILD_LOCATION
-#	CLANG_SOURCE_LOCATION
 #	CPP_ARGS
 #	COMMON_ARGS
 #	APPLE_SDK
@@ -159,11 +177,18 @@ function CallClang()
 	mkdir -p $BUILD_LOCATION
 	for source_file in "${CLANG_SOURCES[@]}"
 	do
+	
+		if [[ "$source_file" =~ "rada_" ]]; then
+			SOURCE_PATH=$RADA_SOURCE_LOCATION
+		else
+			SOURCE_PATH=$RADAUDIO_SOURCE_LOCATION
+		fi
+		
 		if [ "${source_file##*\.}" == "cpp" ]
 		then
-			$CLANG_PATH ${SYSROOT[@]} ${CPP_ARGS} ${CLANG_ARGS[@]} ${COMMON_ARGS[@]} -o $BUILD_LOCATION/${source_file%.*}.o $CLANG_SOURCE_LOCATION/$source_file
+			$CLANG_PATH ${SYSROOT[@]} ${CPP_ARGS} ${CLANG_ARGS[@]} ${COMMON_ARGS[@]} ${EXTRA_INCLUDES[@]} -o $BUILD_LOCATION/${source_file%.*}.o $SOURCE_PATH/$source_file
 		else
-			$CLANG_PATH ${SYSROOT[@]} ${CLANG_ARGS[@]} ${COMMON_ARGS[@]} -o $BUILD_LOCATION/${source_file%.*}.o $CLANG_SOURCE_LOCATION/$source_file
+			$CLANG_PATH ${SYSROOT[@]} ${CLANG_ARGS[@]} ${COMMON_ARGS[@]} ${EXTRA_INCLUDES[@]} -o $BUILD_LOCATION/${source_file%.*}.o $SOURCE_PATH/$source_file
 		fi
 
 		OUTPUT+=( $BUILD_LOCATION/${source_file%.*}.o )
@@ -176,7 +201,6 @@ function CallClang()
 #
 #
 echo Rad Audio Decoder...
-CLANG_SOURCE_LOCATION=$RADAUDIO_SOURCE_LOCATION
 
 #
 # OSX
@@ -195,6 +219,7 @@ CLANG_ARGS=(${OSX_X64_ARGS[@]} ${AVX2_ARGS[@]})
 CLANG_SOURCES=(${RADAUDIO_SOURCES_AVX2[@]})
 CallClang
 
+rm -f $OUTPUT_LOCATION/libradaudio_decoder_osx_x64_static.a
 ar rcs $OUTPUT_LOCATION/libradaudio_decoder_osx_x64_static.a ${OUTPUT[@]}
 
 # ---- OSX ARM64 ----
@@ -205,9 +230,11 @@ CLANG_ARGS=(${OSX_ARM64_ARGS[@]})
 CLANG_SOURCES=(${RADAUDIO_SOURCES[@]})
 CallClang
 
+rm -f $OUTPUT_LOCATION/libradaudio_decoder_osx_arm64_static.a
 ar rcs $OUTPUT_LOCATION/libradaudio_decoder_osx_arm64_static.a ${OUTPUT[@]}
 
 # ---- LIPO ----
+rm -f $OUTPUT_LOCATION/libradaudio_decoder_osx.a
 lipo -create $OUTPUT_LOCATION/libradaudio_decoder_osx_x64_static.a $OUTPUT_LOCATION/libradaudio_decoder_osx_arm64_static.a -output $OUTPUT_LOCATION/libradaudio_decoder_osx.a
 rm $OUTPUT_LOCATION/libradaudio_decoder_osx_arm64_static.a
 rm $OUTPUT_LOCATION/libradaudio_decoder_osx_x64_static.a
@@ -226,6 +253,7 @@ BUILD_LOCATION=$BUILD_ROOT/radaudio_ios_arm64
 CLANG_ARGS=(${IOS_ARM64_ARGS[@]})
 CallClang
 
+rm -f $OUTPUT_LOCATION/libradaudio_decoder_ios.a
 ar rcs $OUTPUT_LOCATION/libradaudio_decoder_ios.a ${OUTPUT[@]}
 
 # ---- iOS Sim x64 ----
@@ -237,6 +265,7 @@ BUILD_LOCATION=$BUILD_ROOT/radaudio_iossim_x64
 CLANG_ARGS=(${IOS_SIM_ARGS_X64[@]})
 CallClang
 
+rm -f $OUTPUT_LOCATION/libradaudio_decoder_iossim_x64_static.a
 ar rcs $OUTPUT_LOCATION/libradaudio_decoder_iossim_x64_static.a ${OUTPUT[@]}
 
 # ---- iOS Sim ARM64 ----
@@ -246,9 +275,11 @@ BUILD_LOCATION=$BUILD_ROOT/radaudio_iossim_arm64
 CLANG_ARGS=(${IOS_SIM_ARGS_ARM64[@]})
 CallClang
 
+rm -f $OUTPUT_LOCATION/libradaudio_decoder_iossim_arm64_static.a
 ar rcs $OUTPUT_LOCATION/libradaudio_decoder_iossim_arm64_static.a ${OUTPUT[@]}
 
 # ---- LIPO SIM ----
+rm -f $OUTPUT_LOCATION/libradaudio_decoder_iossim.a
 lipo -create $OUTPUT_LOCATION/libradaudio_decoder_iossim_x64_static.a $OUTPUT_LOCATION/libradaudio_decoder_iossim_arm64_static.a -output $OUTPUT_LOCATION/libradaudio_decoder_iossim.a
 rm $OUTPUT_LOCATION/libradaudio_decoder_iossim_x64_static.a
 rm $OUTPUT_LOCATION/libradaudio_decoder_iossim_arm64_static.a
@@ -266,6 +297,7 @@ BUILD_LOCATION=$BUILD_ROOT/radaudio_tvos_arm64
 CLANG_ARGS=(${TVOS_ARM64_ARGS[@]})
 CallClang
 
+rm -f $OUTPUT_LOCATION/libradaudio_decoder_tvos.a
 ar rcs $OUTPUT_LOCATION/libradaudio_decoder_tvos.a ${OUTPUT[@]}
 
 #
@@ -280,6 +312,7 @@ BUILD_LOCATION=$BUILD_ROOT/radaudio_visionos_arm64
 CLANG_ARGS=(${VISIONOS_ARM64_ARGS[@]})
 CallClang
 
+rm -f $OUTPUT_LOCATION/libradaudio_decoder_visionos.a
 ar rcs $OUTPUT_LOCATION/libradaudio_decoder_visionos.a ${OUTPUT[@]}
 
 # ---- sim ----
@@ -291,6 +324,7 @@ BUILD_LOCATION=$BUILD_ROOT/radaudio_visionossim_arm64
 CLANG_ARGS=(${VISIONOS_SIM_ARM64_ARGS[@]})
 CallClang
 
+rm -f $OUTPUT_LOCATION/libradaudio_decoder_visionossim.a
 ar rcs $OUTPUT_LOCATION/libradaudio_decoder_visionossim.a ${OUTPUT[@]}
 
 
@@ -315,6 +349,7 @@ CLANG_ARGS=(${OSX_X64_ARGS[@]})
 CLANG_SOURCES=(${RADAUDIO_ENCODER_SOURCES[@]})
 CallClang
 
+rm -f $OUTPUT_LOCATION/libradaudio_encoder_osx_x64_static.a
 ar rcs $OUTPUT_LOCATION/libradaudio_encoder_osx_x64_static.a ${OUTPUT[@]}
 
 # ---- OSX ARM64 ----
@@ -325,168 +360,15 @@ CLANG_ARGS=(${OSX_ARM64_ARGS[@]})
 CLANG_SOURCES=(${RADAUDIO_ENCODER_SOURCES[@]})
 CallClang
 
+rm -f $OUTPUT_LOCATION/libradaudio_encoder_osx_arm64_static.a
 ar rcs $OUTPUT_LOCATION/libradaudio_encoder_osx_arm64_static.a ${OUTPUT[@]}
 
 # ---- LIPO ----
+rm -f $OUTPUT_LOCATION/libradaudio_encoder_osx.a
 lipo -create $OUTPUT_LOCATION/libradaudio_encoder_osx_x64_static.a $OUTPUT_LOCATION/libradaudio_encoder_osx_arm64_static.a -output $OUTPUT_LOCATION/libradaudio_encoder_osx.a
 rm $OUTPUT_LOCATION/libradaudio_encoder_osx_arm64_static.a
 rm $OUTPUT_LOCATION/libradaudio_encoder_osx_x64_static.a
 
 
-#
-#
-# ----- RADA DECODE -------------------------------------
-#
-#
-CLANG_SOURCES=(${RADA_SOURCES[@]})
-CLANG_SOURCE_LOCATION=$RADA_SOURCE_LOCATION
-echo RadA Decoder
-
-#
-# OSX
-#
-echo ...OSX
-
-# ---- OSX X64 ----
-OUTPUT=()
-BUILD_LOCATION=$BUILD_ROOT/rada_osx_x64
-CLANG_ARGS=(${OSX_X64_ARGS[@]})
-CallClang
-
-ar rcs $OUTPUT_LOCATION/librada_decode_osx_x64_static.a ${OUTPUT[@]}
-
-# ---- OSX ARM64 ----
-OUTPUT=()
-BUILD_LOCATION=$BUILD_ROOT/rada_osx_arm64
-CLANG_ARGS=(${OSX_ARM64_ARGS[@]})
-CallClang
-
-ar rcs $OUTPUT_LOCATION/librada_decode_osx_arm64_static.a ${OUTPUT[@]}
-
-# ---- LIPO ----
-lipo -create $OUTPUT_LOCATION/librada_decode_osx_x64_static.a $OUTPUT_LOCATION/librada_decode_osx_arm64_static.a -output $OUTPUT_LOCATION/librada_decode_osx.a
-rm $OUTPUT_LOCATION/librada_decode_osx_x64_static.a
-rm $OUTPUT_LOCATION/librada_decode_osx_arm64_static.a
-
-
-#
-# iOS
-#
-echo ...IOS
-APPLE_SDK=iphoneos
-
-# ---- iOS ARM64 ----
-OUTPUT=()
-BUILD_LOCATION=$BUILD_ROOT/rada_ios_arm64
-
-CLANG_ARGS=(${IOS_ARM64_ARGS[@]})
-CallClang
-
-ar rcs $OUTPUT_LOCATION/librada_decode_ios.a ${OUTPUT[@]}
-
-# ---- iOS Sim x64 ----
-APPLE_SDK=iphonesimulator
-OUTPUT=()
-BUILD_LOCATION=$BUILD_ROOT/rada_iossim_x64
-
-CLANG_ARGS=(${IOS_SIM_ARGS_X64[@]})
-CallClang
-
-ar rcs $OUTPUT_LOCATION/librada_decode_iossim_x64_static.a ${OUTPUT[@]}
-
-# ---- iOS Sim ARM64 ----
-OUTPUT=()
-BUILD_LOCATION=$BUILD_ROOT/rada_iossim_arm64
-
-CLANG_ARGS=(${IOS_SIM_ARGS_ARM64[@]})
-CallClang
-
-ar rcs $OUTPUT_LOCATION/librada_decode_iossim_arm64_static.a ${OUTPUT[@]}
-
-
-# ---- LIPO SIM ----
-lipo -create $OUTPUT_LOCATION/librada_decode_iossim_x64_static.a $OUTPUT_LOCATION/librada_decode_iossim_arm64_static.a -output $OUTPUT_LOCATION/librada_decode_iossim.a
-rm $OUTPUT_LOCATION/librada_decode_iossim_x64_static.a
-rm $OUTPUT_LOCATION/librada_decode_iossim_arm64_static.a
-
-
-#
-# TVOS
-#
-echo TVOS
-APPLE_SDK=appletvos
-
-OUTPUT=()
-BUILD_LOCATION=$BUILD_ROOT/rada_tvos_arm64
-
-CLANG_ARGS=(${TVOS_ARM64_ARGS[@]})
-CallClang
-
-ar rcs $OUTPUT_LOCATION/librada_decode_tvos.a ${OUTPUT[@]}
-
-
-#
-# VisionOS
-#
-echo VisionOS
-APPLE_SDK=xros
-
-OUTPUT=()
-BUILD_LOCATION=$BUILD_ROOT/rada_visionos_arm64
-
-CLANG_ARGS=(${VISIONOS_ARM64_ARGS[@]})
-CallClang
-
-ar rcs $OUTPUT_LOCATION/librada_decode_visionos.a ${OUTPUT[@]}
-
-# ---- sim ----
-APPLE_SDK=xrsimulator
-
-OUTPUT=()
-BUILD_LOCATION=$BUILD_ROOT/rada_visionossim_arm64
-
-CLANG_ARGS=(${VISIONOS_SIM_ARM64_ARGS[@]})
-CallClang
-
-ar rcs $OUTPUT_LOCATION/librada_decode_visionossim.a ${OUTPUT[@]}
-
-
-#
-#
-# ----- RADA ENCODE -------------------------------------
-#
-#
-CLANG_SOURCES=(${RADA_ENCODE_SOURCES[@]})
-CLANG_SOURCE_LOCATION=$RADA_SOURCE_LOCATION
-
-echo RadA Encode
-
-#
-# OSX
-#
-APPLE_SDK=macosx
-
-# ---- OSX X64 ----
-OUTPUT=()
-BUILD_LOCATION=$BUILD_ROOT/rada_encode_osx_x64
-CLANG_ARGS=(${OSX_X64_ARGS[@]})
-CallClang
-
-ar rcs $OUTPUT_LOCATION/librada_encode_osx_x64_static.a ${OUTPUT[@]}
-
-# ---- OSX ARM64 ----
-OUTPUT=()
-BUILD_LOCATION=$BUILD_ROOT/rada_encode_osx_arm64
-CLANG_ARGS=(${OSX_ARM64_ARGS[@]})
-CallClang
-
-ar rcs $OUTPUT_LOCATION/librada_encode_osx_arm64_static.a ${OUTPUT[@]}
-
-# ---- LIPO ----
-lipo -create $OUTPUT_LOCATION/librada_encode_osx_x64_static.a $OUTPUT_LOCATION/librada_encode_osx_arm64_static.a -output $OUTPUT_LOCATION/librada_encode_osx.a
-rm $OUTPUT_LOCATION/librada_encode_osx_x64_static.a
-rm $OUTPUT_LOCATION/librada_encode_osx_arm64_static.a
-
-
 
 echo Done.
diff --git a/Engine/Source/Runtime/RadAudioCodec/SDK/Include/rada_decode.h b/Engine/Source/Runtime/RadAudioCodec/SDK/Include/rada_decode.h
--- a/Engine/Source/Runtime/RadAudioCodec/SDK/Include/rada_decode.h
+++ b/Engine/Source/Runtime/RadAudioCodec/SDK/Include/rada_decode.h
@@ -22,6 +22,8 @@ struct RadAContainer;
     Encoded blocks.
 */
 
+// RADA compatibility version - if these are the same the exports are named the same so we expect
+// that if the linker selects a different copy they all work.
 #define RADA_LIBRARY_VERSION    1
 
 #ifndef RR_STRING_JOIN3
@@ -53,6 +55,7 @@ struct RadAContainer;
 #define RadAExamineBlock            RADA_NAME(RadAExamineBlock)
 #define RadANotifySeek              RADA_NAME(RadANotifySeek)
 #define RadADecodeBlock             RADA_NAME(RadADecodeBlock)
+#define RadAInterleave              RADA_NAME(RadAInterleave)
 
 // The number of bytes the seek table takes up on disk. It exists after
 // RadABytesToOpen, and before RadABytesToFirstBlock
@@ -233,4 +236,9 @@ constexpr uint32_t RadADecodeBlock_MaxOutputFrames = 1024;
 int16_t RadADecodeBlock(RadAContainer* InContainer,
     const uint8_t* InBuffer, size_t InBufferLenBytes, 
     float* InOutputBuffer, size_t InOutBufferStrideInFloats, 
-    size_t* OutConsumedBytes);
\ No newline at end of file
+    size_t* OutConsumedBytes);
+
+// utility function to interleave N buffers of float data into one output buffer of S16 (scales and clamps from -1.0 to 1 into -32768 to 32767)
+#define RADA_HAS_INTERLEAVING
+void  RadAInterleave( int16_t * output, float const ** inputs, uint32_t input_count, uint32_t samples );
+
diff --git a/Engine/Source/Runtime/RadAudioCodec/SDK/Include/rada_encode.h b/Engine/Source/Runtime/RadAudioCodec/SDK/Include/rada_encode.h
--- a/Engine/Source/Runtime/RadAudioCodec/SDK/Include/rada_encode.h
+++ b/Engine/Source/Runtime/RadAudioCodec/SDK/Include/rada_encode.h
@@ -41,7 +41,15 @@ typedef void RadACompressFreeFnType(void* Ptr);
 
 #define RadAErrorString     RADA_NAME(RadAErrorString)
 #define EncodeRadAFile      RADA_NAME(EncodeRadAFile)
+#define RadAGetBuildVersion RADA_NAME(RadAGetBuildVersion)
 
+// The build version of the rad audio library in it's entirety. This is distinct from the
+// compilation version that determine linker compatibility. This is incremented each time
+// the library is released
+#define RADA_BUILD_VERSION (0x00000001U)
+
+// Return the build version that the library was compiled with. This should match RADA_BUILD_VERSION!
+uint32_t RadAGetBuildVersion();
 
 // Pass in one of the above values to get a human readable error, returned
 // as a utf8 string.
diff --git a/Engine/Source/Runtime/RadAudioCodec/SDK/Src/RadA/rada_decode.cpp b/Engine/Source/Runtime/RadAudioCodec/SDK/Src/RadA/rada_decode.cpp
--- a/Engine/Source/Runtime/RadAudioCodec/SDK/Src/RadA/rada_decode.cpp
+++ b/Engine/Source/Runtime/RadAudioCodec/SDK/Src/RadA/rada_decode.cpp
@@ -475,11 +475,23 @@ static size_t RadASeekTableLookupInternal(const RadAFileHeader* InFileHeader, co
 
         if (OutFrameAtLocation)
             *OutFrameAtLocation = RadAEvaluateSeekTableSamplePosition(InFileHeader, InSeekHeader, InSeekTableData, seek_table_qwords, block_index);;
-        if (OutFrameBlockSize)
-            *OutFrameBlockSize = InFileHeader->max_block_size;
 
         uint64_t byte_result = RadAEvaluateSeekTableBytePosition(InFileHeader, InSeekHeader, InSeekTableData, seek_table_qwords, block_index);
 
+        if (OutFrameBlockSize)
+        {
+            // We are somewhere in this seek block, but we don't know exactly where, data wise. Since we want to _guarantee_ that the frame can
+            // be reached, the only option is to report the entire seek table entry.
+            //
+            size_t next_entry_bytes = 0;
+            if ((block_index + 1) < InFileHeader->seek_table_entry_count)
+                next_entry_bytes = RadAEvaluateSeekTableBytePosition(InFileHeader, InSeekHeader, InSeekTableData, seek_table_qwords, block_index + 1);
+            else
+                next_entry_bytes = InFileHeader->file_size - RadAGetBytesToOpen(InFileHeader); // we are doing math in seek table space...
+
+            *OutFrameBlockSize = next_entry_bytes - byte_result;
+        }
+
         return byte_result + RadAGetBytesToOpen(InFileHeader);
     }
 }
@@ -699,3 +711,12 @@ int16_t RadADecodeBlock(RadAContainer* rada,
     *consumed_bytes = input_buffer - input_reservoir;
     return decoded_samples;
 }
+
+#ifdef RADAUDIO_DECODER_HAS_INTERLEAVING
+
+void  RadAInterleave( int16_t * output, float const ** inputs, uint32_t input_count, uint32_t samples )
+{
+    RadAudioInterleave( output, inputs, input_count, samples );
+}
+
+#endif
diff --git a/Engine/Source/Runtime/RadAudioCodec/SDK/Src/RadA/rada_encode.cpp b/Engine/Source/Runtime/RadAudioCodec/SDK/Src/RadA/rada_encode.cpp
--- a/Engine/Source/Runtime/RadAudioCodec/SDK/Src/RadA/rada_encode.cpp
+++ b/Engine/Source/Runtime/RadAudioCodec/SDK/Src/RadA/rada_encode.cpp
@@ -551,6 +551,16 @@ static void LinearRegressionDeltasToBits(UInt64Array* resolved_deltas,
 {
     resolved_deltas->count = 0;
 
+    if (count == 0)
+    {
+        out_regression->intercept = 0;
+        out_regression->slope_numerator = 0;
+        out_regression->slope_denomenator = 1;
+        out_regression->bias = 0;
+        *out_bits_per_entry = 0;
+        return;
+    }
+
     // We are delivered deltas, but we want to regress on the resolved positions, so sum in to 
     // our temp space
     uint64_t delta_sum = 0;
@@ -1223,9 +1233,12 @@ uint8_t EncodeRadAFile(
             memcpy(Current, &Header, sizeof(RadAFileHeader));
             Current += sizeof(RadAFileHeader);
 
-            RadASeekTableHeader seek_table_header = seek_table_info.get_header();
-            memcpy(Current, &seek_table_header, sizeof(RadASeekTableHeader));
-            Current += sizeof(RadASeekTableHeader);
+            if (GenerateSeekTable)
+            {
+                RadASeekTableHeader seek_table_header = seek_table_info.get_header();
+                memcpy(Current, &seek_table_header, sizeof(RadASeekTableHeader));
+                Current += sizeof(RadASeekTableHeader);
+            }
 
             for (uint8_t stream_index = 0; stream_index < StreamCount; stream_index++)
             {
@@ -1282,6 +1295,11 @@ const char* RadAErrorString(uint8_t InError)
     return "Invalid RadA Error Code";
 }
 
+uint32_t RadAGetBuildVersion()
+{
+	return RADA_BUILD_VERSION;
+}
+
 #if 0
 
 /* to build this fuzzed:
diff --git a/Engine/Source/Runtime/RadAudioCodec/SDK/Src/RadAudio/cpux86.cpp b/Engine/Source/Runtime/RadAudioCodec/SDK/Src/RadAudio/cpux86.cpp
--- a/Engine/Source/Runtime/RadAudioCodec/SDK/Src/RadAudio/cpux86.cpp
+++ b/Engine/Source/Runtime/RadAudioCodec/SDK/Src/RadAudio/cpux86.cpp
@@ -13,6 +13,16 @@
 
 	#if _MSC_VER >= 1500 // VC++2008 or later
 	#define HAVE_CPUIDEX
+
+	static inline U64 xgetbv(U32 xcr)
+	{
+		return _xgetbv(xcr);
+	}
+	#else
+	static inline U64 xgetbv(U32 xcr)
+	{
+		return 0;
+	}
 	#endif
 
 #else
@@ -33,6 +43,13 @@
 
 	#define HAVE_CPUIDEX
 	#define __cpuid(out, leaf_id) __cpuidex(out, leaf_id, 0)
+
+	static inline U64 xgetbv(U32 xcr)
+	{
+		U32 lo, hi;
+		__asm__ __volatile__("xgetbv" : "=a"(lo), "=d"(hi) : "c"(xcr));
+		return ((U64)hi << 32) | lo;
+	}
 #endif // _MSC_VER or not
 
 
@@ -74,11 +91,13 @@ extern "C" void rrCPUx86_detect()
 	if (cpuid_info[2] & (1u<< 9))	features |= RRX86_CPU_SSSE3;
 	if (cpuid_info[2] & (1u<<19))	features |= RRX86_CPU_SSE41;
 	if (cpuid_info[2] & (1u<<20))	features |= RRX86_CPU_SSE42;
-	if (cpuid_info[2] & (1u<<28))	features |= RRX86_CPU_AVX;
-	if (cpuid_info[2] & (1u<<29))	features |= RRX86_CPU_F16C;
 
-	// We don't have a feature flag we report for this, but we do use it later
+	// Used to compute other feature flags
 	bool has_popcnt = (cpuid_info[2] & (1u<<23)) != 0;
+	bool has_osxsave = (cpuid_info[2] & (1u<<27)) != 0;
+	bool has_cpu_avx = (cpuid_info[2] & (1u<<28)) != 0;
+	bool has_cpu_f16c = (cpuid_info[2] & (1u<<29)) != 0;
+
 	if (has_popcnt) features |= RRX86_CPU_POPCNT;
 
 	if (is_amd)
@@ -93,6 +112,24 @@ extern "C" void rrCPUx86_detect()
 			features |= RRX86_CPU_AMD_ZEN;
 	}
 
+	// Get XCR0, if available, and determine context save bits
+	U64 xcr0 = 0;
+	if (has_osxsave)
+	{
+		xcr0 = xgetbv(0);
+	}
+
+	// YMM register saving and ZMM/opmask register saving support
+	bool has_os_avx_support = (xcr0 & 6) == 6;
+	bool has_os_avx512_support = (xcr0 & 0xe6) == 0xe6;
+
+	// AVX support requires both CPU and OS support, and gates some other extensions
+	if (has_os_avx_support)
+	{
+		if (has_cpu_avx)	features |= RRX86_CPU_AVX;
+		if (has_cpu_f16c)	features |= RRX86_CPU_F16C;
+	}
+
 #ifdef HAVE_CPUIDEX
 	if (max_leaf >= 7)
 	{
@@ -107,7 +144,7 @@ extern "C" void rrCPUx86_detect()
 		// Also only report AVX or the BMIs if POPCNT is present; all processors I know of
 		// have either both or neither, and it's convenient for us to be able to assume
 		// that either BMI1/BMI2 or AVX2 implies POPCNT.
-		if ((features & RRX86_CPU_AVX) && has_popcnt)
+		if (has_cpu_avx && has_os_avx_support && has_popcnt)
 		{
 			if (cpuid_info[1] & (1u<<3))	features |= RRX86_CPU_BMI1;
 			if (cpuid_info[1] & (1u<<8))	features |= RRX86_CPU_BMI2;
@@ -123,17 +160,20 @@ extern "C" void rrCPUx86_detect()
 			if ((cpuid_info[1] & avx2_bits) == avx2_bits)
 				features |= RRX86_CPU_AVX2;
 
-			// For us to report AVX512, we want the Skylake feature set
-			const U32 avx512_bits = (1u << 31) /* AVX512VL */ | (1u << 30) /* AVX512BW */ | (1u << 17) /* AVX512DQ */ | (1u << 16) /* AVX512F */;
-			if ((cpuid_info[1] & avx512_bits) == avx512_bits)
-				features |= RRX86_CPU_AVX512;
-
-			// Use the VBMI2 bit (set on ICL+) to set the PREFER512 flag. This is available
-			// on a generation of cores where AVX-512 has no major clock penalty anymore so
-			// whether to use AVX-512 or not is a much more straightforward calculation,
-			// and not so dependent on what else is running at the same time.
-			if (cpuid_info[2] & (1u << 6))
-				features |= RRX86_CPU_PREFER512;
+			if (has_os_avx512_support)
+			{
+				// For us to report AVX512, we want the Skylake feature set
+				const U32 avx512_bits = (1u << 31) /* AVX512VL */ | (1u << 30) /* AVX512BW */ | (1u << 17) /* AVX512DQ */ | (1u << 16) /* AVX512F */;
+				if ((cpuid_info[1] & avx512_bits) == avx512_bits)
+					features |= RRX86_CPU_AVX512;
+
+				// Use the VBMI2 bit (set on ICL+) to set the PREFER512 flag. This is available
+				// on a generation of cores where AVX-512 has no major clock penalty anymore so
+				// whether to use AVX-512 or not is a much more straightforward calculation,
+				// and not so dependent on what else is running at the same time.
+				if (cpuid_info[2] & (1u << 6))
+					features |= RRX86_CPU_PREFER512;
+			}
 		}
 	}
 #endif
diff --git a/Engine/Source/Runtime/RadAudioCodec/SDK/Src/RadAudio/egttypes.h b/Engine/Source/Runtime/RadAudioCodec/SDK/Src/RadAudio/egttypes.h
--- a/Engine/Source/Runtime/RadAudioCodec/SDK/Src/RadAudio/egttypes.h
+++ b/Engine/Source/Runtime/RadAudioCodec/SDK/Src/RadAudio/egttypes.h
@@ -1,4 +1,6 @@
+//===============================================
 // Copyright Epic Games Tools, LLC. All Rights Reserved.
+//===============================================
 
 #ifndef __EGTTYPESH__
 #define __EGTTYPESH__
@@ -7,6 +9,8 @@
 #define __RADTYPESH__
 #define __RADRR_COREH__ // block old rr_core
 
+#define RADCOPYRIGHT "Copyright (C) 1994-2025, Epic Games Tools, LLC"
+
 #if !defined(__RADRES__) // don't include anything for resource compiles
 
 //  __RAD32__ means at least 32 bit code (always defined)
diff --git a/Engine/Source/Runtime/RadAudioCodec/SDK/Src/RadAudio/radaudio_decoder.h b/Engine/Source/Runtime/RadAudioCodec/SDK/Src/RadAudio/radaudio_decoder.h
--- a/Engine/Source/Runtime/RadAudioCodec/SDK/Src/RadAudio/radaudio_decoder.h
+++ b/Engine/Source/Runtime/RadAudioCodec/SDK/Src/RadAudio/radaudio_decoder.h
@@ -25,8 +25,12 @@ typedef struct
 
 RADDEFSTART
 
+// RADAUDIO compatibility version - if these are the same the exports are named the same so we expect
+// that if the linker selects a different copy they all work.
 #define RADAUDIO_DECODER_LIBRARY_VERSION    1
 
+#define RADAUDIO_DECODER_HAS_INTERLEAVING
+
 // The max size of a stream header. It could be smaller but this is a safe size
 // for buffers.
 #define RADAUDIO_STREAM_HEADER_SIZE   128
@@ -52,6 +56,11 @@ RADDEFSTART
 #define RadAudioDecoderDidSeek               RADAUDIO_DECODER_NAME(RadAudioDecoderDidSeek)
 #define RadAudioDecoderGetChunkLength        RADAUDIO_DECODER_NAME(RadAudioDecoderGetChunkLength)
 
+#define RadAudioInterleave                   RADAUDIO_DECODER_NAME(RadAudioInterleave)
+#define RadAudio1ChToS16                     RADAUDIO_DECODER_NAME(RadAudio1ChToS16)
+#define RadAudioInterleave2ChToS16           RADAUDIO_DECODER_NAME(RadAudioInterleave2ChToS16)
+#define RadAudioInterleave4ChToS16           RADAUDIO_DECODER_NAME(RadAudioInterleave4ChToS16)
+
 // RadAudioDecoderMemoryRequired()
 //
 // returns the amount of memory that needs to be passed to Open().
@@ -127,6 +136,18 @@ RADDEFFUNC int RadAudioDecoderGetChunkLength(RadAudioDecoder *radaudio_decoder,
 // inform the decoder that you performed a seek operation on the input stream,
 // such that the next decoder call is not continuous with the previous one.
 RADDEFFUNC void RadAudioDecoderDidSeek(RadAudioDecoder *radaudio_decoder);
+
+
+
+// utility function to interleave N buffers of float data into one output buffer of S16 (scales and clamps from -1.0 to 1 into -32768 to 32767)
+RADDEFFUNC void  RadAudioInterleave( S16 * output, float const ** inputs, U32 input_count, U32 samples );
+
+// low level direct routines, if these return 0, then they converted NOTHING (the sample count is too small for even one SIMD round).
+//    if they return 1, then the ENTIRE buffer was converted.
+RADDEFFUNC int   RadAudio1ChToS16( S16 * output, float const * input, U32 samples );
+RADDEFFUNC int   RadAudioInterleave2ChToS16( S16 * output, float const * input_left, float const * input_right, U32 samples );
+RADDEFFUNC int   RadAudioInterleave4ChToS16( S16 * output, float const * input_left, float const * input_right, float const * input_back_left, float const * input_back_right, U32 samples );
+
 RADDEFEND
 
 #endif
diff --git a/Engine/Source/Runtime/RadAudioCodec/SDK/Src/RadAudio/radaudio_encoder.h b/Engine/Source/Runtime/RadAudioCodec/SDK/Src/RadAudio/radaudio_encoder.h
--- a/Engine/Source/Runtime/RadAudioCodec/SDK/Src/RadAudio/radaudio_encoder.h
+++ b/Engine/Source/Runtime/RadAudioCodec/SDK/Src/RadAudio/radaudio_encoder.h
@@ -28,6 +28,8 @@ typedef struct
 
 RADDEFSTART
 
+// RADAUDIO compatibility version - if these are the same the exports are named the same so we expect
+// that if the linker selects a different copy they all work.
 #define RADAUDIO_ENC_LIBRARY_VERSION    1
 
 #ifndef RR_STRING_JOIN3
diff --git a/Engine/Source/Runtime/RadAudioCodec/SDK/Src/RadAudio/radaudio_interleave.c b/Engine/Source/Runtime/RadAudioCodec/SDK/Src/RadAudio/radaudio_interleave.c
new file mode 100644
--- /dev/null
+++ b/Engine/Source/Runtime/RadAudioCodec/SDK/Src/RadAudio/radaudio_interleave.c
@@ -0,0 +1,260 @@
+#include "radaudio_decoder.h"
+#include "radaudio_sse.h"
+
+#if defined(_x86_64) || defined( __x86_64__ ) || defined( _M_X64 ) || defined(__x86_64) || defined(_M_AMD64) || defined(__SSE2__) || defined(DO_BUILD_SSE4)
+
+  #include <emmintrin.h>
+  #define simdf __m128
+  #define simdi32 __m128i
+  #define simdi16 __m128i
+  #define simdf_float_count 4
+  #define SIMD_PATH
+
+  #define simdf_load( ptr ) _mm_loadu_ps(ptr)
+  #define simdf_mul_constant( src, scalar ) _mm_mul_ps( (src), _mm_set1_ps( scalar ) )
+  #define simdf_min_constant( src, scalar ) _mm_min_ps( (src), _mm_set1_ps( scalar ) )
+  #define simdf_max_constant( src, scalar ) _mm_max_ps( (src), _mm_set1_ps( scalar ) )
+  #define simdf_to_simdi32( src ) _mm_cvttps_epi32( src )
+
+  #define simdi32_packto16( rega, regb ) _mm_packs_epi32( (rega), (regb) )
+
+  #define simdi16_store( ptr, values ) _mm_storeu_si128( (__m128i *) (ptr), (values) )
+
+  static RADFORCEINLINE void simdi16_store2_interleaved(S16 * ptr, simdi16 rega, simdi16 regb)
+  {
+    simdi16_store(ptr + 0, _mm_unpacklo_epi16(rega, regb));
+    simdi16_store(ptr + 8, _mm_unpackhi_epi16(rega, regb));
+  }
+
+  static RADFORCEINLINE void simdi16_store4_interleaved(S16 * ptr, simdi16 rega, simdi16 regb, simdi16 regc, simdi16 regd)
+  {
+    simdi16 ac0 = _mm_unpacklo_epi16(rega, regc);
+    simdi16 ac1 = _mm_unpackhi_epi16(rega, regc);
+    simdi16 bd0 = _mm_unpacklo_epi16(regb, regd);
+    simdi16 bd1 = _mm_unpackhi_epi16(regb, regd);
+
+    simdi16_store(ptr +  0, _mm_unpacklo_epi16(ac0, bd0));
+    simdi16_store(ptr +  8, _mm_unpackhi_epi16(ac0, bd0));
+    simdi16_store(ptr + 16, _mm_unpacklo_epi16(ac1, bd1));
+    simdi16_store(ptr + 24, _mm_unpackhi_epi16(ac1, bd1));
+  }
+
+#elif defined( _M_ARM64 ) || defined( __aarch64__ ) || defined( __arm64__ ) || ((__ARM_NEON_FP & 4) != 0) || defined(__ARM_NEON__) || defined(DO_BUILD_NEON) || defined(__RADARM64__)
+
+  #include <arm_neon.h>
+
+  #define simdf float32x4_t
+  #define simdi32 int32x4_t
+  #define simdi16 int16x8_t
+  #define simdf_float_count 4
+  #define SIMD_PATH
+
+  #define simdf_load( ptr ) vld1q_f32(ptr)
+  #define simdf_mul_constant( src, scalar ) vmulq_n_f32( src, scalar )
+  #define simdf_min_constant( src, scalar ) vminq_f32( src, vdupq_n_f32( scalar ) )
+  #define simdf_max_constant( src, scalar ) vmaxq_f32( src, vdupq_n_f32( scalar ) )
+  #define simdf_to_simdi32( src ) vcvtq_s32_f32( src )
+
+  #define simdi32_packto16( rega, regb ) vcombine_s16( vmovn_s32( rega ), vmovn_s32( regb ) )
+
+  #define simdi16_store( ptr, values ) vst1q_s16( (ptr), (values) )
+
+  static RADFORCEINLINE void simdi16_store2_interleaved(S16 * ptr, simdi16 rega, simdi16 regb)
+  {
+    int16x8x2_t values = { rega, regb };
+    vst2q_s16(ptr, values);
+  }
+
+  static RADFORCEINLINE void simdi16_store4_interleaved(S16 * ptr, simdi16 rega, simdi16 regb, simdi16 regc, simdi16 regd)
+  {
+    int16x8x4_t values = { rega, regb, regc, regd };
+    vst4q_s16(ptr, values);
+  }
+
+#endif
+
+#ifdef SIMD_PATH
+
+static RADFORCEINLINE simdi32 simd_load_and_convert(float const * ptr)
+{
+  simdf floats = simdf_load(ptr);
+  simdf scaled = simdf_mul_constant(floats, 32768.0f);
+  simdf clamped_lo = simdf_max_constant(scaled, -32768.0f);
+  simdf clamped_hi = simdf_min_constant(clamped_lo, 32767.0f);
+  simdi32 integer = simdf_to_simdi32(clamped_hi);
+  return integer;
+}
+
+static RADFORCEINLINE simdi16 simd_load_convert_narrow(float const * ptr)
+{
+  simdi32 lo = simd_load_and_convert(ptr);
+  simdi32 hi = simd_load_and_convert(ptr + simdf_float_count);
+  return simdi32_packto16(lo, hi);
+}
+
+#endif
+
+int   RadAudio1ChToS16( S16 * output, float const * input, U32 samples )
+{
+  #ifndef SIMD_PATH
+
+  return 0;
+
+  #else
+
+  S16 * end_output;
+  float const * end_input_m8;
+  
+  // if buffer is so small that we can't do anything, fail
+  if ( samples < simdf_float_count*2 )
+    return 0;
+
+  // as long as we have at least simdf_float_count*2 samples, this loop handles all of them (no remenants).
+  end_input_m8 = input + samples - simdf_float_count*2;
+  end_output = output + samples - simdf_float_count*2;
+  for(;;)
+  {
+    simdi16 o = simd_load_convert_narrow(input);
+    simdi16_store(output, o);
+    input += simdf_float_count*2;
+    output += simdf_float_count*2;
+    if ( output <= end_output )
+      continue;
+    if ( output == ( end_output + simdf_float_count*2 ) )
+      return 1;
+    output = end_output;     // backup and do last couple
+    input = end_input_m8;
+  }
+  #endif
+}
+
+
+int   RadAudioInterleave2ChToS16( S16 * output, float const * input_left, float const * input_right, U32 samples )
+{
+  #ifndef SIMD_PATH
+
+  return 0;
+
+  #else
+
+  S16 * end_output;
+  float const * end_left_m8;
+  float const * end_right_m8;
+  
+  // if buffer is so small that we can't do anything, fail
+  if ( samples < simdf_float_count*2 )
+    return 0;
+
+  // as long as we have at least simdf_float_count*2 samples, this loop handles all of them (no remenants).
+  end_left_m8  = input_left  + samples - simdf_float_count*2;
+  end_right_m8 = input_right + samples - simdf_float_count*2;
+  end_output = output + samples*2 - simdf_float_count*2*2; // 2 for two simd regs, and 2 for stereo
+  for(;;)
+  {
+    simdi16 ol = simd_load_convert_narrow(input_left);
+    simdi16 or = simd_load_convert_narrow(input_right);
+    simdi16_store2_interleaved(output, ol, or);
+    input_left += simdf_float_count*2;
+    input_right += simdf_float_count*2;
+    output += simdf_float_count*2*2;
+    if ( output <= end_output )
+      continue;
+    if ( output == ( end_output + simdf_float_count*2*2 ) )
+      return 1;
+    output = end_output;     // backup and do last couple
+    input_left  = end_left_m8;
+    input_right = end_right_m8;
+  }
+
+  #endif
+}
+
+int   RadAudioInterleave4ChToS16( S16 * output, float const * input_left, float const * input_right, float const * input_back_left, float const * input_back_right, U32 samples )
+{
+  #ifndef SIMD_PATH
+
+  return 0;
+
+  #else
+
+  S16 * end_output;
+  float const * end_fleft_m8;
+  float const * end_fright_m8;
+  float const * end_bleft_m8;
+  float const * end_bright_m8;
+
+  // if buffer is so small that we can't do anything, fail
+  if ( samples < simdf_float_count*2 )
+    return 0;
+
+  // as long as we have at least simdf_float_count*2 samples, this loop handles all of them (no remenants).
+  end_fleft_m8  = input_left  + samples - simdf_float_count*2;
+  end_fright_m8 = input_right + samples - simdf_float_count*2;
+  end_bleft_m8  = input_back_left  + samples - simdf_float_count*2;
+  end_bright_m8 = input_back_right + samples - simdf_float_count*2;
+  end_output = output + samples*4 - simdf_float_count*2*4; // 2 for two simd regs, and 4 for quad-chan
+  for(;;)
+  {
+    simdi16 ofl = simd_load_convert_narrow(input_left);
+    simdi16 ofr = simd_load_convert_narrow(input_right);
+    simdi16 obl = simd_load_convert_narrow(input_back_left);
+    simdi16 obr = simd_load_convert_narrow(input_back_right);
+    simdi16_store4_interleaved(output, ofl, ofr, obl, obr);
+
+    input_left += simdf_float_count*2;
+    input_right += simdf_float_count*2;
+    input_back_left += simdf_float_count*2;
+    input_back_right += simdf_float_count*2;
+    output += simdf_float_count*2*4;
+    if ( output <= end_output )
+      continue;
+    if ( output == ( end_output + simdf_float_count*2*4 ) )
+      return 1;
+    output = end_output;     // backup and do last couple
+    input_left  = end_fleft_m8;
+    input_right = end_fright_m8;
+    input_back_left  = end_bleft_m8;
+    input_back_right = end_bright_m8;
+  }
+
+  #endif
+}
+
+void   RadAudioInterleave( S16 * output, float const ** inputs, U32 input_count, U32 samples )
+{
+  U32 SampleIdx;
+  switch( input_count )
+  {
+    case 0:
+      break;
+    case 1:
+      if ( RadAudio1ChToS16( output, inputs[0], samples ) )
+        return;
+      break;
+    case 2:
+      if ( RadAudioInterleave2ChToS16( output, inputs[0], inputs[1], samples ) )
+        return;
+      break;
+    case 4:
+      if ( RadAudioInterleave4ChToS16( output, inputs[0], inputs[1], inputs[2], inputs[3], samples ) )
+        return;
+      break;
+    default:
+      break;
+  }
+
+  for (SampleIdx = 0; SampleIdx < samples; SampleIdx++)
+  {
+    U32 ChannelIdx;
+    for (ChannelIdx = 0; ChannelIdx < input_count; ChannelIdx++)
+    {
+      float InBufferFloat = inputs[ChannelIdx][SampleIdx] * 32768.0f;
+      if ( InBufferFloat < -32768.0f )
+        InBufferFloat = -32768.0f;
+      if ( InBufferFloat > 32767.0f )
+        InBufferFloat = 32767.0f;
+
+      *output++ = (S16)InBufferFloat;
+    }
+  }
+}
diff --git a/Engine/Source/Runtime/RadAudioCodec/SDK/Src/RadAudio/radaudio_mdct.cpp b/Engine/Source/Runtime/RadAudioCodec/SDK/Src/RadAudio/radaudio_mdct.cpp
--- a/Engine/Source/Runtime/RadAudioCodec/SDK/Src/RadAudio/radaudio_mdct.cpp
+++ b/Engine/Source/Runtime/RadAudioCodec/SDK/Src/RadAudio/radaudio_mdct.cpp
@@ -4,6 +4,7 @@
 #endif
 #include <stdio.h>
 #include <stdlib.h>
+#include <stddef.h>
 #include <assert.h>
 #include <math.h>
 
diff --git a/Engine/Source/Runtime/RadAudioCodec/SDK/Src/RadAudio/rrCore.h b/Engine/Source/Runtime/RadAudioCodec/SDK/Src/RadAudio/rrCore.h
--- a/Engine/Source/Runtime/RadAudioCodec/SDK/Src/RadAudio/rrCore.h
+++ b/Engine/Source/Runtime/RadAudioCodec/SDK/Src/RadAudio/rrCore.h
@@ -1,4 +1,6 @@
+/// ========================================================================
 // Copyright Epic Games Tools, LLC. All Rights Reserved. Global types header file
+// ========================================================================
 
 #if !defined(__RADRR_CORE2H__)
 #define __RADRR_CORE2H__
