copyright fix

#rnx

diff --git a/Engine/Source/Runtime/RadAudioCodec/SDK/Src/RadAudio/egttypes.h b/Engine/Source/Runtime/RadAudioCodec/SDK/Src/RadAudio/egttypes.h
--- a/Engine/Source/Runtime/RadAudioCodec/SDK/Src/RadAudio/egttypes.h
+++ b/Engine/Source/Runtime/RadAudioCodec/SDK/Src/RadAudio/egttypes.h
@@ -7,8 +7,6 @@
 #define __RADTYPESH__
 #define __RADRR_COREH__ // block old rr_core
 
-#define RADCOPYRIGHT "Copyright (C) 1994-2025, Epic Games Tools, LLC"
-
 #if !defined(__RADRES__) // don't include anything for resource compiles
 
 //  __RAD32__ means at least 32 bit code (always defined)
