Use new Rad Audio interleavers.

#rnx

diff --git a/Engine/Source/Runtime/RadAudioCodec/Module/Private/RadAudioInfo.cpp b/Engine/Source/Runtime/RadAudioCodec/Module/Private/RadAudioInfo.cpp
--- a/Engine/Source/Runtime/RadAudioCodec/Module/Private/RadAudioInfo.cpp
+++ b/Engine/Source/Runtime/RadAudioCodec/Module/Private/RadAudioInfo.cpp
@@ -427,10 +427,15 @@ FDecodeResult FRadAudioInfo::Decode(const uint8* CompressedData, const int32 Com
 		// Interleave in to the output buffer, and convert to 16 bit.
 		// The buffers are [0..DecodeResult...1024..1024+DecodeResult...n*1024..n*1024+DecodeResult]
 		// We want: [0.1...n, DecodeResult...DecodeResult+n]
-		// \todo optimize this disaster. we're waiting on some sane shared lib stuff rather than writing it here.
-		// as 1) the shared lib float conversion uses 32767 instead of 32768, which means that you can end up with
-		// < -1.0f in the buffer and 2) we can't interleave with SIMD as there's only VectorDeinterleave rather than
-		// interleave.
+#ifdef RADA_HAS_INTERLEAVING
+		const float* DeinterleavedSources[RADA_MAX_CHANNELS];
+		for (int32 ChannelIdx = 0; ChannelIdx < NumChannels; ChannelIdx++)
+		{
+			DeinterleavedSources[ChannelIdx] = DeinterleavedDecodeBuffer.GetData() + RadADecodeBlock_MaxOutputFrames * ChannelIdx + DecodeResultOffset;
+		}
+		RadAInterleave(InterleaveDestination, DeinterleavedSources, NumChannels, DecodeResult);
+
+#else
 		for (int32 ChannelIdx = 0; ChannelIdx < NumChannels; ChannelIdx++)
 		{
 			const float* InBuffer = DeinterleavedDecodeBuffer.GetData() + RadADecodeBlock_MaxOutputFrames*ChannelIdx + DecodeResultOffset;
@@ -450,6 +455,7 @@ FDecodeResult FRadAudioInfo::Decode(const uint8* CompressedData, const int32 Com
 				InterleaveDestination[ChannelIdx + (SampleIdx * NumChannels)] = (int16)InBufferFloat;
 			}
 		}
+#endif
 
 		if (InterleaveDestination == (int16*)OutputReservoir.GetData())
 		{
